RC protocols auto binding

diff --git a/Engine/Plugins/VirtualProduction/RemoteControl/Source/RemoteControl/Private/RemoteControlProtocolBinding.cpp b/Engine/Plugins/VirtualProduction/RemoteControl/Source/RemoteControl/Private/RemoteControlProtocolBinding.cpp
--- a/Engine/Plugins/VirtualProduction/RemoteControl/Source/RemoteControl/Private/RemoteControlProtocolBinding.cpp
+++ b/Engine/Plugins/VirtualProduction/RemoteControl/Source/RemoteControl/Private/RemoteControlProtocolBinding.cpp
@@ -498,6 +498,25 @@ bool FRemoteControlProtocolEntity::ApplyProtocolValueToProperty(double InProtoco
 	return bSuccess;
 }
 
+ERCBindingStatus FRemoteControlProtocolEntity::ToggleBindingStatus()
+{
+	if (BindingStatus == ERCBindingStatus::Awaiting)
+	{
+		BindingStatus = ERCBindingStatus::Bound;
+	}
+	else if (BindingStatus == ERCBindingStatus::Bound || BindingStatus == ERCBindingStatus::Unassigned)
+	{
+		BindingStatus = ERCBindingStatus::Awaiting;
+	}
+
+	return BindingStatus;
+}
+
+void FRemoteControlProtocolEntity::ResetDefaultBindingState()
+{
+	BindingStatus = ERCBindingStatus::Unassigned;
+}
+
 bool FRemoteControlProtocolEntity::GetInterpolatedPropertyBuffer(FProperty* InProperty, double InProtocolValue, TArray<uint8>& OutBuffer) const
 {
 	OutBuffer.Empty();
diff --git a/Engine/Plugins/VirtualProduction/RemoteControl/Source/RemoteControl/Public/RemoteControlProtocolBinding.h b/Engine/Plugins/VirtualProduction/RemoteControl/Source/RemoteControl/Public/RemoteControlProtocolBinding.h
--- a/Engine/Plugins/VirtualProduction/RemoteControl/Source/RemoteControl/Public/RemoteControlProtocolBinding.h
+++ b/Engine/Plugins/VirtualProduction/RemoteControl/Source/RemoteControl/Public/RemoteControlProtocolBinding.h
@@ -19,6 +19,18 @@ struct FRCObjectReference;
 
 using FGetProtocolMappingCallback = TFunctionRef<void(FRemoteControlProtocolMapping&)>;
 
+/**
+ * Status of the binding 
+ */
+UENUM()
+enum class ERCBindingStatus : uint8
+{
+	Unassigned,
+	Awaiting ,
+	Bound
+};
+
+
 /*
  * Mapping of the range of the values for the protocol
  * This class holds a generic range buffer.
@@ -353,6 +365,18 @@ public:
 	/** Get Size of the range property value */
 	uint8 GetRangePropertySize() const;
 
+	/** Get current binding state of the entity */
+	ERCBindingStatus GetBindingStatus() const { return BindingStatus; }
+
+	/**
+	 * Toggle the binding
+	 * @return Current binding state
+	 */
+	ERCBindingStatus ToggleBindingStatus();
+
+	/** Reset the binding set to default */ 
+	void ResetDefaultBindingState();
+
 private:
 
 	/**
@@ -384,6 +408,11 @@ protected:
 	 */
 	UPROPERTY()
 	TSet<FRemoteControlProtocolMapping> Mappings;
+
+private:
+	/** Binding status of this protocol entity */
+	UPROPERTY()
+	ERCBindingStatus BindingStatus = ERCBindingStatus::Unassigned;
 };
 
 /**
diff --git a/Engine/Plugins/VirtualProduction/RemoteControl/Source/RemoteControlProtocolWidgets/Private/RemoteControlProtocolWidgetsModule.cpp b/Engine/Plugins/VirtualProduction/RemoteControl/Source/RemoteControlProtocolWidgets/Private/RemoteControlProtocolWidgetsModule.cpp
--- a/Engine/Plugins/VirtualProduction/RemoteControl/Source/RemoteControlProtocolWidgets/Private/RemoteControlProtocolWidgetsModule.cpp
+++ b/Engine/Plugins/VirtualProduction/RemoteControl/Source/RemoteControlProtocolWidgets/Private/RemoteControlProtocolWidgetsModule.cpp
@@ -11,14 +11,26 @@ DEFINE_LOG_CATEGORY(LogRemoteControlProtocolWidgets);
 TSharedRef<SWidget> FRemoteControlProtocolWidgetsModule::GenerateDetailsForEntity(URemoteControlPreset* InPreset, const FGuid& InFieldId, const EExposedFieldType& InFieldType)
 {
 	check(InPreset);
-
+ 
+	ResetProtocolBindingList();
+ 
 	// Currently only supports Properties
 	if(InFieldId.IsValid() && InFieldType == EExposedFieldType::Property)
 	{
-		return SNew(SRCProtocolBindingList, FProtocolEntityViewModel::Create(InPreset, InFieldId));
+		return SAssignNew(RCProtocolBindingList, SRCProtocolBindingList, FProtocolEntityViewModel::Create(InPreset, InFieldId));
 	}
-
+ 
 	return SNullWidget::NullWidget;
 }
 
 IMPLEMENT_MODULE(FRemoteControlProtocolWidgetsModule, RemoteControlProtocolWidgets);
+ 
+void FRemoteControlProtocolWidgetsModule::ResetProtocolBindingList()
+{
+	RCProtocolBindingList = nullptr;
+}
+ 
+TSharedPtr<IRCProtocolBindingList> FRemoteControlProtocolWidgetsModule::GetProtocolBindingList() const 
+{
+	return RCProtocolBindingList;
+}
\ No newline at end of file
diff --git a/Engine/Plugins/VirtualProduction/RemoteControl/Source/RemoteControlProtocolWidgets/Private/RemoteControlProtocolWidgetsModule.h b/Engine/Plugins/VirtualProduction/RemoteControl/Source/RemoteControlProtocolWidgets/Private/RemoteControlProtocolWidgetsModule.h
--- a/Engine/Plugins/VirtualProduction/RemoteControl/Source/RemoteControlProtocolWidgets/Private/RemoteControlProtocolWidgetsModule.h
+++ b/Engine/Plugins/VirtualProduction/RemoteControl/Source/RemoteControlProtocolWidgets/Private/RemoteControlProtocolWidgetsModule.h
@@ -6,11 +6,19 @@
 
 REMOTECONTROLPROTOCOLWIDGETS_API DECLARE_LOG_CATEGORY_EXTERN(LogRemoteControlProtocolWidgets, Log, All);
 
+class IRCProtocolBindingList;
 class URemoteControlPreset;
 
 class FRemoteControlProtocolWidgetsModule : public IRemoteControlProtocolWidgetsModule
 {
 public:
-	/** Creates a widget for the given Preset Field */
+	//~ Begin IRemoteControlProtocolWidgetsModule Interface
 	virtual TSharedRef<SWidget> GenerateDetailsForEntity(URemoteControlPreset* InPreset, const FGuid& InFieldId, const EExposedFieldType& InFieldType) override;
+	virtual void ResetProtocolBindingList() override;
+	virtual TSharedPtr<IRCProtocolBindingList> GetProtocolBindingList() const override;
+	//~ End IRemoteControlProtocolWidgetsModule Interface
+
+private:
+	/** Binding list public interface instance */
+	TSharedPtr<IRCProtocolBindingList> RCProtocolBindingList;
 };
diff --git a/Engine/Plugins/VirtualProduction/RemoteControl/Source/RemoteControlProtocolWidgets/Private/Widgets/SRCProtocolBinding.cpp b/Engine/Plugins/VirtualProduction/RemoteControl/Source/RemoteControlProtocolWidgets/Private/Widgets/SRCProtocolBinding.cpp
--- a/Engine/Plugins/VirtualProduction/RemoteControl/Source/RemoteControlProtocolWidgets/Private/Widgets/SRCProtocolBinding.cpp
+++ b/Engine/Plugins/VirtualProduction/RemoteControl/Source/RemoteControlProtocolWidgets/Private/Widgets/SRCProtocolBinding.cpp
@@ -2,6 +2,8 @@
 
 #include "SRCProtocolBinding.h"
 
+#include "EditorFontGlyphs.h"
+
 #include "SDropTarget.h"
 #include "SRCProtocolRangeList.h"
 #include "ViewModels/ProtocolBindingViewModel.h"
@@ -18,6 +20,8 @@ void SRCProtocolBinding::Construct(const FArguments& InArgs, const TSharedRef<ST
 	ViewModel = InViewModel;
 	PrimaryColumnSizeData = InArgs._PrimaryColumnSizeData;
 	SecondaryColumnSizeData = InArgs._SecondaryColumnSizeData;
+	OnStartRecording = InArgs._OnStartRecording;
+	OnStopRecording = InArgs._OnStopRecording;
 	
 	TSharedRef<SHorizontalBox> RowBox = SNew(SHorizontalBox);
 
@@ -47,6 +51,25 @@ void SRCProtocolBinding::Construct(const FArguments& InArgs, const TSharedRef<ST
 			.Text(FText::FromString(FString(TEXT("\xf00d"))))
 		]
 	];
+
+	RowBox->AddSlot()
+	.AutoWidth()
+	.VAlign(VAlign_Center)
+	.Padding(0)
+	[
+		SNew(SButton)
+		.ButtonStyle(FEditorStyle::Get(), "HoverHintOnly")
+		.ToolTipText(LOCTEXT("RecordingButtonToolTip", "Status of the protocol entity binding"))
+		.ForegroundColor(FSlateColor::UseForeground())
+		.OnClicked(this, &SRCProtocolBinding::ToggleRecording)
+		.Content()
+		[
+			SNew(STextBlock)
+			.ColorAndOpacity_Raw(this, &SRCProtocolBinding::GetRecordingButtonColor)
+			.Font(FEditorStyle::Get().GetFontStyle("FontAwesome.10"))
+			.Text(FEditorFontGlyphs::Circle)
+		]
+	];
 	
 	STableRow::Construct(
 		STableRow::FArguments()
@@ -109,4 +132,57 @@ FReply SRCProtocolBinding::OnDelete() const
 	return FReply::Handled();
 }
 
+FReply SRCProtocolBinding::ToggleRecording() const
+{
+	// Binding can't be nullptr
+	FRemoteControlProtocolBinding* Binding = ViewModel->GetBinding();
+	check(Binding);
+
+	if (const TSharedPtr<TStructOnScope<FRemoteControlProtocolEntity>> ProtocolEntity = Binding->GetRemoteControlProtocolEntityPtr())
+	{
+		const ERCBindingStatus BindingStatus = (*ProtocolEntity)->ToggleBindingStatus();
+		if (BindingStatus == ERCBindingStatus::Awaiting)
+		{
+			OnStartRecording.ExecuteIfBound(ProtocolEntity);
+		}
+		else if (BindingStatus == ERCBindingStatus::Bound)
+		{
+			OnStopRecording.ExecuteIfBound(ProtocolEntity);
+		}
+		else
+		{
+			checkNoEntry();
+		}
+	}
+
+	return FReply::Handled();
+}
+
+FSlateColor SRCProtocolBinding::GetRecordingButtonColor() const
+{
+	// Binding can't be nullptr
+	FRemoteControlProtocolBinding* Binding = ViewModel->GetBinding();
+	check(Binding);
+
+	if (const TSharedPtr<TStructOnScope<FRemoteControlProtocolEntity>> ProtocolEntity = Binding->GetRemoteControlProtocolEntityPtr())
+	{
+		const ERCBindingStatus BindingStatus = (*ProtocolEntity)->GetBindingStatus();
+
+		switch (BindingStatus)
+		{
+			case ERCBindingStatus::Awaiting:
+				return FLinearColor::Red;
+			case ERCBindingStatus::Bound:
+				return FLinearColor::Green;
+			case ERCBindingStatus::Unassigned:
+				return FLinearColor::Gray;
+			default:
+				checkNoEntry();
+		}
+	}
+
+	ensure(false);
+	return FLinearColor::Black;
+}
+
 #undef LOCTEXT_NAMESPACE
diff --git a/Engine/Plugins/VirtualProduction/RemoteControl/Source/RemoteControlProtocolWidgets/Private/Widgets/SRCProtocolBinding.h b/Engine/Plugins/VirtualProduction/RemoteControl/Source/RemoteControlProtocolWidgets/Private/Widgets/SRCProtocolBinding.h
--- a/Engine/Plugins/VirtualProduction/RemoteControl/Source/RemoteControlProtocolWidgets/Private/Widgets/SRCProtocolBinding.h
+++ b/Engine/Plugins/VirtualProduction/RemoteControl/Source/RemoteControlProtocolWidgets/Private/Widgets/SRCProtocolBinding.h
@@ -5,10 +5,14 @@
 #include "CoreMinimal.h"
 
 #include "SRCProtocolShared.h"
+#include "UObject/StructOnScope.h"
 #include "ViewModels/ProtocolBindingViewModel.h"
 #include "Widgets/DeclarativeSyntaxSupport.h"
 #include "Widgets/Views/STableRow.h"
 
+struct FRemoteControlProtocolEntity;
+struct FSlateColor;
+
 class FProtocolBindingViewModel;
 class FProtocolBindingViewModel;
 class FProtocolRangeViewModel;
@@ -21,10 +25,18 @@ class STableViewBase;
 class REMOTECONTROLPROTOCOLWIDGETS_API SRCProtocolBinding : public STableRow<TSharedPtr<FProtocolBindingViewModel>>
 {
 public:
+	/** Delegate when start recording input from protocol */
+	DECLARE_DELEGATE_OneParam(FOnStartRecording, TSharedPtr<TStructOnScope<FRemoteControlProtocolEntity>> ProtocolEntity)
+
+	/** Delegate when stop recording input from protocol */
+	DECLARE_DELEGATE_OneParam(FOnStopRecording, TSharedPtr<TStructOnScope<FRemoteControlProtocolEntity>> ProtocolEntity)
+	
 	SLATE_BEGIN_ARGS(SRCProtocolBinding) 
 	{}
 		SLATE_ARGUMENT(TSharedPtr<RemoteControlProtocolWidgetUtils::FPropertyViewColumnSizeData>, PrimaryColumnSizeData)
 		SLATE_ARGUMENT(TSharedPtr<RemoteControlProtocolWidgetUtils::FPropertyViewColumnSizeData>, SecondaryColumnSizeData)
+		SLATE_EVENT( FOnStartRecording, OnStartRecording )
+		SLATE_EVENT( FOnStopRecording, OnStopRecording )
 	SLATE_END_ARGS()
 
 	void Construct(const FArguments& InArgs, const TSharedRef<STableViewBase>& InOwnerTableView, const TSharedRef<FProtocolBindingViewModel>& InViewModel);
@@ -32,6 +44,12 @@ public:
 private:
 	FReply OnDelete() const;
 
+	/** Toggle recording button handler */
+	FReply ToggleRecording() const;
+
+	/** Get Recording button color based on binding status of protocol entity */
+	FSlateColor GetRecordingButtonColor() const;
+
 private:
 	TSharedPtr<FProtocolBindingViewModel> ViewModel;
 
@@ -42,4 +60,10 @@ private:
 
 	/** Container used by all secondary splitters in the details view, so that they move in sync */
 	TSharedPtr<RemoteControlProtocolWidgetUtils::FPropertyViewColumnSizeData> SecondaryColumnSizeData;
+
+	/** Start recording delegate instance */
+	FOnStartRecording OnStartRecording;
+
+	/** Stop recording delegate instance */
+	FOnStopRecording OnStopRecording;
 };
diff --git a/Engine/Plugins/VirtualProduction/RemoteControl/Source/RemoteControlProtocolWidgets/Private/Widgets/SRCProtocolBindingList.cpp b/Engine/Plugins/VirtualProduction/RemoteControl/Source/RemoteControlProtocolWidgets/Private/Widgets/SRCProtocolBindingList.cpp
--- a/Engine/Plugins/VirtualProduction/RemoteControl/Source/RemoteControlProtocolWidgets/Private/Widgets/SRCProtocolBindingList.cpp
+++ b/Engine/Plugins/VirtualProduction/RemoteControl/Source/RemoteControlProtocolWidgets/Private/Widgets/SRCProtocolBindingList.cpp
@@ -188,16 +188,30 @@ void SRCProtocolBindingList::Construct(const FArguments& InArgs, TSharedRef<FPro
 
 SRCProtocolBindingList::~SRCProtocolBindingList()
 {
+	for (const TSharedPtr<TStructOnScope<FRemoteControlProtocolEntity>>& RemoteControlProtocolEntity : AwaitingProtocolEntities)
+	{
+		if (FRemoteControlProtocolEntity* RemoteControlProtocolEntityPtr = RemoteControlProtocolEntity->Get())
+		{
+			if (RemoteControlProtocolEntityPtr->GetBindingStatus() == ERCBindingStatus::Awaiting)
+			{
+				RemoteControlProtocolEntityPtr->ResetDefaultBindingState();
+			}
+		}
+	}
+
+	AwaitingProtocolEntities.Empty();
 }
 
-TSharedRef<SRCProtocolBinding> SRCProtocolBindingList::ConstructBindingWidget(const TSharedRef<STableViewBase>& InOwnerTable, TSharedPtr<FProtocolBindingViewModel> InViewModel) const
+TSharedRef<SRCProtocolBinding> SRCProtocolBindingList::ConstructBindingWidget(const TSharedRef<STableViewBase>& InOwnerTable, TSharedPtr<FProtocolBindingViewModel> InViewModel)
 {
 	return SNew(SRCProtocolBinding, InOwnerTable, InViewModel.ToSharedRef())
 		.PrimaryColumnSizeData(PrimaryColumnSizeData)
-		.SecondaryColumnSizeData(SecondaryColumnSizeData);
+		.SecondaryColumnSizeData(SecondaryColumnSizeData)
+		.OnStartRecording(this, &SRCProtocolBindingList::OnStartRecording)
+		.OnStopRecording(this, &SRCProtocolBindingList::OnStopRecording);
 }
 
-TSharedRef<ITableRow> SRCProtocolBindingList::OnGenerateRow(TSharedPtr<FProtocolBindingViewModel> InViewModel, const TSharedRef<STableViewBase>& OwnerTable) const
+TSharedRef<ITableRow> SRCProtocolBindingList::OnGenerateRow(TSharedPtr<FProtocolBindingViewModel> InViewModel, const TSharedRef<STableViewBase>& OwnerTable)
 {
 	check(InViewModel.IsValid());
 	TSharedRef<SRCProtocolBinding> BindingWidget = ConstructBindingWidget(OwnerTable, InViewModel);
@@ -235,6 +249,16 @@ bool SRCProtocolBindingList::IsProtocolShown(const FName& InProtocolName)
 	return !GetSettings()->HiddenProtocolTypeNames.Contains(InProtocolName);
 }
 
+void SRCProtocolBindingList::OnStartRecording(TSharedPtr<TStructOnScope<FRemoteControlProtocolEntity>> InEntity)
+{
+	AwaitingProtocolEntities.Add(InEntity);
+}
+
+void SRCProtocolBindingList::OnStopRecording(TSharedPtr<TStructOnScope<FRemoteControlProtocolEntity>> InEntity)
+{
+	AwaitingProtocolEntities.Remove(InEntity);
+}
+
 URemoteControlProtocolWidgetsSettings* SRCProtocolBindingList::GetSettings()
 {
 	if(Settings.IsValid())
diff --git a/Engine/Plugins/VirtualProduction/RemoteControl/Source/RemoteControlProtocolWidgets/Private/Widgets/SRCProtocolBindingList.h b/Engine/Plugins/VirtualProduction/RemoteControl/Source/RemoteControlProtocolWidgets/Private/Widgets/SRCProtocolBindingList.h
--- a/Engine/Plugins/VirtualProduction/RemoteControl/Source/RemoteControlProtocolWidgets/Private/Widgets/SRCProtocolBindingList.h
+++ b/Engine/Plugins/VirtualProduction/RemoteControl/Source/RemoteControlProtocolWidgets/Private/Widgets/SRCProtocolBindingList.h
@@ -4,9 +4,9 @@
 
 #include "CoreMinimal.h"
 
-
+#include "IRCProtocolBindingList.h"
 #include "SRCProtocolShared.h"
-#include "UObject/StrongObjectPtr.h"
+
 #include "Widgets/SCompoundWidget.h"
 #include "Widgets/Views/SListView.h"
 
@@ -20,7 +20,7 @@ class URemoteControlProtocolWidgetsSettings;
 template <typename ItemType> class SListView;
 
 /** The root view for a given entity. A (vertical) list of bindings, where each binding has a protocol. */
-class REMOTECONTROLPROTOCOLWIDGETS_API SRCProtocolBindingList final : public SCompoundWidget
+class REMOTECONTROLPROTOCOLWIDGETS_API SRCProtocolBindingList final : public SCompoundWidget, public IRCProtocolBindingList
 {
 public:
 	SLATE_BEGIN_ARGS(SRCProtocolBindingList)
@@ -30,10 +30,14 @@ public:
 	void Construct(const FArguments& InArgs, TSharedRef<FProtocolEntityViewModel> InViewModel);
 	virtual ~SRCProtocolBindingList();
 
+	//~ Begin IRCProtocolBindingList Interface
+	virtual FRemoteControlProtocolEntitySet GetAwaitingProtocolEntities() const override { return AwaitingProtocolEntities; }
+	//~ End IRCProtocolBindingList Interface
+
 private:
-	TSharedRef<SRCProtocolBinding> ConstructBindingWidget(const TSharedRef<STableViewBase>&, TSharedPtr<FProtocolBindingViewModel> InViewModel) const;
+	TSharedRef<SRCProtocolBinding> ConstructBindingWidget(const TSharedRef<STableViewBase>&, TSharedPtr<FProtocolBindingViewModel> InViewModel);
 	
-	TSharedRef<ITableRow> OnGenerateRow(TSharedPtr<FProtocolBindingViewModel> InViewModel, const TSharedRef<STableViewBase>& OwnerTable) const;
+	TSharedRef<ITableRow> OnGenerateRow(TSharedPtr<FProtocolBindingViewModel> InViewModel, const TSharedRef<STableViewBase>& OwnerTable);
 
 	bool CanAddProtocol();
 
@@ -53,6 +57,12 @@ private:
 	float OnGetSecondaryRightColumnWidth() const { return SecondaryColumnWidth; }
 	void OnSetSecondaryColumnWidth(float InWidth) { SecondaryColumnWidth = InWidth; }
 
+	/** Start recording incoming protocol message handler */
+	void OnStartRecording(TSharedPtr<TStructOnScope<FRemoteControlProtocolEntity>> InEntity);
+
+	/** Stop recording incoming protocol message handler */
+	void OnStopRecording(TSharedPtr<TStructOnScope<FRemoteControlProtocolEntity>> InEntity);
+
 	/** Get (mutable) module settings */
 	URemoteControlProtocolWidgetsSettings* GetSettings();
 
@@ -88,4 +98,7 @@ private:
 
 	/** Reference to (mutable) settings class for this module */
 	TWeakObjectPtr<URemoteControlProtocolWidgetsSettings> Settings;
+
+	/** Set of protocol entities with awaiting state and waiting for binding */
+	FRemoteControlProtocolEntitySet AwaitingProtocolEntities;
 };
diff --git a/Engine/Plugins/VirtualProduction/RemoteControl/Source/RemoteControlProtocolWidgets/Public/IRCProtocolBindingList.h b/Engine/Plugins/VirtualProduction/RemoteControl/Source/RemoteControlProtocolWidgets/Public/IRCProtocolBindingList.h
new file mode 100644
--- /dev/null
+++ b/Engine/Plugins/VirtualProduction/RemoteControl/Source/RemoteControlProtocolWidgets/Public/IRCProtocolBindingList.h
@@ -0,0 +1,23 @@
+﻿// Copyright Epic Games, Inc. All Rights Reserved.
+
+#pragma once
+
+#include "CoreMinimal.h"
+#include "UObject/StructOnScope.h"
+
+struct FRemoteControlProtocolEntity;
+
+/**
+ * Interface class for RCProtocolBindingList widget
+ */
+class IRCProtocolBindingList
+{
+public:
+	using FRemoteControlProtocolEntitySet = TSet<TSharedPtr<TStructOnScope<FRemoteControlProtocolEntity>>>;
+	
+	/** Virtual destructor */
+	virtual ~IRCProtocolBindingList() = default;
+
+	/** Get the set of entities which is awaiting state and waiting for binding */ 
+	virtual FRemoteControlProtocolEntitySet GetAwaitingProtocolEntities() const = 0;
+};
\ No newline at end of file
diff --git a/Engine/Plugins/VirtualProduction/RemoteControl/Source/RemoteControlProtocolWidgets/Public/IRemoteControlProtocolWidgetsModule.h b/Engine/Plugins/VirtualProduction/RemoteControl/Source/RemoteControlProtocolWidgets/Public/IRemoteControlProtocolWidgetsModule.h
--- a/Engine/Plugins/VirtualProduction/RemoteControl/Source/RemoteControlProtocolWidgets/Public/IRemoteControlProtocolWidgetsModule.h
+++ b/Engine/Plugins/VirtualProduction/RemoteControl/Source/RemoteControlProtocolWidgets/Public/IRemoteControlProtocolWidgetsModule.h
@@ -6,6 +6,7 @@
 #include "Modules/ModuleManager.h"
 #include "RemoteControlField.h"
 
+class IRCProtocolBindingList;
 class SWidget;
 class URemoteControlPreset;
 class SWidget;
@@ -14,6 +15,24 @@ class SWidget;
 class IRemoteControlProtocolWidgetsModule : public IModuleInterface
 {
 public:
+	/**
+	* Singleton-like access to this module's interface.  This is just for convenience!
+	* Beware of calling this during the shutdown phase, though.  Your module might have been unloaded already.
+	*
+	* @return Returns singleton instance, loading the module on demand if needed
+	*/
+	static IRemoteControlProtocolWidgetsModule& Get()
+	{
+		static const FName ModuleName = "RemoteControlProtocolWidgets";
+		return FModuleManager::LoadModuleChecked<IRemoteControlProtocolWidgetsModule>(ModuleName);
+	}
+	
 	/** Creates a widget for the given Preset Field and FieldType */
 	virtual TSharedRef<SWidget> GenerateDetailsForEntity(URemoteControlPreset* InPreset, const FGuid& InFieldId, const EExposedFieldType& InFieldType = EExposedFieldType::Invalid) = 0;
+
+	/** Reset protocol binding widget */
+	virtual void ResetProtocolBindingList() = 0;
+	
+	/** Get the binding list public reference */
+	virtual TSharedPtr<IRCProtocolBindingList> GetProtocolBindingList() const = 0;
 };
diff --git a/Engine/Plugins/VirtualProduction/RemoteControl/Source/RemoteControlUI/Private/UI/SRemoteControlPanel.cpp b/Engine/Plugins/VirtualProduction/RemoteControl/Source/RemoteControlUI/Private/UI/SRemoteControlPanel.cpp
--- a/Engine/Plugins/VirtualProduction/RemoteControl/Source/RemoteControlUI/Private/UI/SRemoteControlPanel.cpp
+++ b/Engine/Plugins/VirtualProduction/RemoteControl/Source/RemoteControlUI/Private/UI/SRemoteControlPanel.cpp
@@ -326,6 +326,10 @@ SRemoteControlPanel::~SRemoteControlPanel()
 
 	// Clear the log
 	FRemoteControlLogger::Get().ClearLog();
+
+	// Remove protocol bindings
+	IRemoteControlProtocolWidgetsModule& ProtocolWidgetsModule = FModuleManager::LoadModuleChecked<IRemoteControlProtocolWidgetsModule>("RemoteControlProtocolWidgets");
+	ProtocolWidgetsModule.ResetProtocolBindingList();	
 }
 
 void SRemoteControlPanel::PostUndo(bool bSuccess)
diff --git a/Engine/Plugins/VirtualProduction/RemoteControlProtocolDMX/Source/RemoteControlProtocolDMX/Private/RemoteControlProtocolDMX.cpp b/Engine/Plugins/VirtualProduction/RemoteControlProtocolDMX/Source/RemoteControlProtocolDMX/Private/RemoteControlProtocolDMX.cpp
--- a/Engine/Plugins/VirtualProduction/RemoteControlProtocolDMX/Source/RemoteControlProtocolDMX/Private/RemoteControlProtocolDMX.cpp
+++ b/Engine/Plugins/VirtualProduction/RemoteControlProtocolDMX/Source/RemoteControlProtocolDMX/Private/RemoteControlProtocolDMX.cpp
@@ -10,6 +10,11 @@
 #include "IO/DMXPortManager.h"
 #include "Library/DMXEntityFixtureType.h"
 
+#if WITH_EDITOR
+#include "IRCProtocolBindingList.h"
+#include "IRemoteControlProtocolWidgetsModule.h"
+#endif
+
 #define LOCTEXT_NAMESPACE "RemoteControlProtocolDMX"
 
 const FName FRemoteControlProtocolDMX::ProtocolName = TEXT("DMX");
@@ -143,6 +148,10 @@ void FRemoteControlProtocolDMX::OnEndFrame()
 				const FDMXSignalSharedPtr& LastSignalPtr = DMXProtocolEntity->LastSignalPtr;
 				const int32 DMXOffset = DMXProtocolEntity->ExtraSetting.StartingChannel - 1;
 				check(DMXOffset >= 0 && DMXOffset < DMX_UNIVERSE_SIZE);
+
+#if WITH_EDITOR
+				ProcessAutoBinding(ProtocolEntity);
+#endif
 				
 				ProcessAndApplyProtocolValue(LastSignalPtr, DMXOffset, ProtocolEntity);
 			}
@@ -182,6 +191,77 @@ void FRemoteControlProtocolDMX::ProcessAndApplyProtocolValue(const FDMXSignalSha
 	}
 }
 
+#if WITH_EDITOR
+void FRemoteControlProtocolDMX::ProcessAutoBinding(const FRemoteControlProtocolEntityPtr& InProtocolEntityPtr)
+{
+	// Bind only in Editor
+	if (!GIsEditor)
+	{
+		return;
+	}
+
+	// Check if entity is valid
+	if (!InProtocolEntityPtr.IsValid())
+	{
+		return;
+	}
+
+	FRemoteControlDMXProtocolEntity* DMXProtocolEntity = InProtocolEntityPtr->CastChecked<FRemoteControlDMXProtocolEntity>();
+	
+	// Assign binding
+	IRemoteControlProtocolWidgetsModule& RCWidgetsModule = IRemoteControlProtocolWidgetsModule::Get();
+	const TSharedPtr<IRCProtocolBindingList> RCProtocolBindingList = RCWidgetsModule.GetProtocolBindingList();
+	if (RCProtocolBindingList.IsValid())
+	{
+		if (DMXProtocolEntity->GetBindingStatus() == ERCBindingStatus::Awaiting)
+		{
+			const TArray<uint8>& ChannelData = DMXProtocolEntity->LastSignalPtr->ChannelData;
+			int32 FoundChannelDifference = -1;
+			uint8 FoundChannelDifferenceValue = 0;
+
+			if (CacheUniverseDMXBuffer.Num())
+			{
+				if (ensure(ChannelData.Num() == CacheUniverseDMXBuffer.Num()))
+				{
+					// Compare buffers
+					for (int32 ChannelValueIndex = 0; ChannelValueIndex < ChannelData.Num(); ++ChannelValueIndex)
+					{
+						const uint8 SignalChannelValue = ChannelData[ChannelValueIndex];
+						if (CacheUniverseDMXBuffer[ChannelValueIndex] != SignalChannelValue)
+						{
+							FoundChannelDifference = ChannelValueIndex;
+							FoundChannelDifferenceValue = SignalChannelValue;
+							break;
+						}
+					}	
+				}
+			}
+
+			if  (FoundChannelDifference >= 0)
+			{
+				Unbind(InProtocolEntityPtr);
+				const int32 FinalChannelValue = FoundChannelDifference + 1;
+				DMXProtocolEntity->ExtraSetting.StartingChannel = FinalChannelValue;
+				Bind(InProtocolEntityPtr);
+
+				// Print to log
+				const FDMXSignalSharedPtr& LastSignalPtr = DMXProtocolEntity->LastSignalPtr;
+				FRemoteControlLogger::Get().Log(ProtocolName, [&LastSignalPtr, FinalChannelValue, FoundChannelDifferenceValue]
+				{
+					return FText::Format(
+						LOCTEXT("DMXEventLog",
+								"AutoBinding new value. ExternUniverseID {0}, Channel {1}, New Value {2}"),
+						LastSignalPtr->ExternUniverseID, FinalChannelValue, FoundChannelDifferenceValue);
+				});
+			}
+
+			// Copy buffer
+			CacheUniverseDMXBuffer = DMXProtocolEntity->LastSignalPtr->ChannelData;
+		}
+	}
+}
+#endif
+
 void FRemoteControlProtocolDMX::UnbindAll()
 {
 	ProtocolsBindings.Empty();
diff --git a/Engine/Plugins/VirtualProduction/RemoteControlProtocolDMX/Source/RemoteControlProtocolDMX/Public/RemoteControlProtocolDMX.h b/Engine/Plugins/VirtualProduction/RemoteControlProtocolDMX/Source/RemoteControlProtocolDMX/Public/RemoteControlProtocolDMX.h
--- a/Engine/Plugins/VirtualProduction/RemoteControlProtocolDMX/Source/RemoteControlProtocolDMX/Public/RemoteControlProtocolDMX.h
+++ b/Engine/Plugins/VirtualProduction/RemoteControlProtocolDMX/Source/RemoteControlProtocolDMX/Public/RemoteControlProtocolDMX.h
@@ -118,10 +118,23 @@ private:
 	 */
 	void ProcessAndApplyProtocolValue(const FDMXSignalSharedPtr& InSignal, int32 InDMXOffset, const FRemoteControlProtocolEntityPtr& InProtocolEntityPtr);
 
+#if WITH_EDITOR
+	/**
+	 * Process the AutoBinding to the Remote Control Entity
+	 * @param InProtocolEntityPtr	Protocol entity pointer
+	 */
+	void ProcessAutoBinding(const FRemoteControlProtocolEntityPtr& InProtocolEntityPtr);
+#endif
+
 private:
 	/** Binding for the DMX protocol */
 	TArray<FRemoteControlProtocolEntityWeakPtr> ProtocolsBindings;
 
+#if WITH_EDITORONLY_DATA
+	/** DMX universe cache buffer.*/
+	TArray<uint8, TFixedAllocator<DMX_UNIVERSE_SIZE>> CacheUniverseDMXBuffer;
+#endif
+
 public:
 	/** DMX protocol name */
 	static const FName ProtocolName;
diff --git a/Engine/Plugins/VirtualProduction/RemoteControlProtocolDMX/Source/RemoteControlProtocolDMX/RemoteControlProtocolDMX.Build.cs b/Engine/Plugins/VirtualProduction/RemoteControlProtocolDMX/Source/RemoteControlProtocolDMX/RemoteControlProtocolDMX.Build.cs
--- a/Engine/Plugins/VirtualProduction/RemoteControlProtocolDMX/Source/RemoteControlProtocolDMX/RemoteControlProtocolDMX.Build.cs
+++ b/Engine/Plugins/VirtualProduction/RemoteControlProtocolDMX/Source/RemoteControlProtocolDMX/RemoteControlProtocolDMX.Build.cs
@@ -21,5 +21,14 @@ public class RemoteControlProtocolDMX : ModuleRules
 				"RemoteControlProtocol",
 			}
 		);
+		
+		if (Target.bBuildEditor)
+		{
+			PrivateDependencyModuleNames.AddRange(
+				new string[] {
+					"RemoteControlProtocolWidgets"
+				}
+			);
+		}
 	}
 }
diff --git a/Engine/Plugins/VirtualProduction/RemoteControlProtocolMIDI/Source/RemoteControlProtocolMIDI/Private/RemoteControlProtocolMIDI.cpp b/Engine/Plugins/VirtualProduction/RemoteControlProtocolMIDI/Source/RemoteControlProtocolMIDI/Private/RemoteControlProtocolMIDI.cpp
--- a/Engine/Plugins/VirtualProduction/RemoteControlProtocolMIDI/Source/RemoteControlProtocolMIDI/Private/RemoteControlProtocolMIDI.cpp
+++ b/Engine/Plugins/VirtualProduction/RemoteControlProtocolMIDI/Source/RemoteControlProtocolMIDI/Private/RemoteControlProtocolMIDI.cpp
@@ -7,6 +7,11 @@
 #include "RemoteControlProtocolMIDIModule.h"
 #include "RemoteControlProtocolMIDISettings.h"
 
+#if WITH_EDITOR
+#include "IRCProtocolBindingList.h"
+#include "IRemoteControlProtocolWidgetsModule.h"
+#endif
+
 #define LOCTEXT_NAMESPACE "FRemoteControlProtocolMIDI"
 
 const FName FRemoteControlProtocolMIDI::ProtocolName = TEXT("MIDI");
@@ -209,6 +214,10 @@ void FRemoteControlProtocolMIDI::OnReceiveEvent(UMIDIDeviceInputController* MIDI
 {
 	const EMIDIEventType MIDIEventType = static_cast<EMIDIEventType>(Type);
 
+#if WITH_EDITOR
+	ProcessAutoBinding(MIDIEventType, Channel, MessageData1);
+#endif
+
 #if WITH_EDITOR
 	{		
 		FRemoteControlLogger::Get().Log(ProtocolName, [Type, Channel, MessageData1, MessageData2]
@@ -276,6 +285,40 @@ void FRemoteControlProtocolMIDI::OnReceiveEvent(UMIDIDeviceInputController* MIDI
 	}
 }
 
+#if WITH_EDITOR
+void FRemoteControlProtocolMIDI::ProcessAutoBinding(EMIDIEventType MIDIEventType, int32 Channel, int32 MessageData1)
+{
+	// Bind only in Editor
+	if (!GIsEditor)
+	{
+		return;
+	}
+		
+	IRemoteControlProtocolWidgetsModule& RCWidgetsModule = IRemoteControlProtocolWidgetsModule::Get();
+	const TSharedPtr<IRCProtocolBindingList> RCProtocolBindingList = RCWidgetsModule.GetProtocolBindingList();
+	if (RCProtocolBindingList.IsValid())
+	{
+		for (const TSharedPtr<TStructOnScope<FRemoteControlProtocolEntity>>& ProtocolEntityPtr : RCProtocolBindingList->GetAwaitingProtocolEntities())
+		{
+			if (ProtocolEntityPtr.IsValid())
+			{
+				if ((*ProtocolEntityPtr)->GetBindingStatus() == ERCBindingStatus::Awaiting)
+				{
+					Unbind(ProtocolEntityPtr);
+					
+					FRemoteControlMIDIProtocolEntity* MIDIProtocolEntity = ProtocolEntityPtr->CastChecked<FRemoteControlMIDIProtocolEntity>();
+					MIDIProtocolEntity->EventType = MIDIEventType;
+					MIDIProtocolEntity->Channel = Channel;
+					MIDIProtocolEntity->MessageData1 = MessageData1;
+
+					Bind(ProtocolEntityPtr);
+				}
+			}	
+		}
+	}
+}
+#endif
+
 void FRemoteControlProtocolMIDI::UnbindAll()
 {
 	MIDIDeviceBindings_NoteOn.Empty();
diff --git a/Engine/Plugins/VirtualProduction/RemoteControlProtocolMIDI/Source/RemoteControlProtocolMIDI/Public/RemoteControlProtocolMIDI.h b/Engine/Plugins/VirtualProduction/RemoteControlProtocolMIDI/Source/RemoteControlProtocolMIDI/Public/RemoteControlProtocolMIDI.h
--- a/Engine/Plugins/VirtualProduction/RemoteControlProtocolMIDI/Source/RemoteControlProtocolMIDI/Public/RemoteControlProtocolMIDI.h
+++ b/Engine/Plugins/VirtualProduction/RemoteControlProtocolMIDI/Source/RemoteControlProtocolMIDI/Public/RemoteControlProtocolMIDI.h
@@ -158,9 +158,28 @@ public:
 	//~ End IRemoteControlProtocol interface
 
 private:
-	/** On receive MIDI buffer callback */
+	/**
+	 * On receive MIDI buffer callback
+	 * @param	MIDIDeviceController	The MIDI Input Device Controller
+	 * @param	Timestamp				The MIDI timestamp
+	 * @param	Channel					The MIDI channel to send 
+	 * @param	MessageData1			The first part of the MIDI data
+	 * @param	MessageData2			The second part of the MIDI data
+	 */
 	void OnReceiveEvent(UMIDIDeviceInputController* MIDIDeviceController, int32 Timestamp, int32 Type, int32 Channel, int32 MessageData1, int32 MessageData2);
 
+#if WITH_EDITOR
+	/**
+	 * Process the AutoBinding to the Remote Control Entity
+	 * @param	MIDIEventType		The event type as specified in the EMIDIEventType struct
+	 * @param	Channel				The MIDI channel 
+	 * @param	MessageData1		The first part of the MIDI data
+	 */
+	void ProcessAutoBinding(EMIDIEventType MIDIEventType, int32 Channel, int32 MessageData1);
+#endif
+
+private:
+
 	/** Binding for ControlChange (11) MIDI protocol */
 	TMap<UMIDIDeviceInputController*, TMap<int32, TArray<FRemoteControlProtocolEntityWeakPtr>>> MIDIDeviceBindings_ControlChange;
 
diff --git a/Engine/Plugins/VirtualProduction/RemoteControlProtocolMIDI/Source/RemoteControlProtocolMIDI/RemoteControlProtocolMIDI.Build.cs b/Engine/Plugins/VirtualProduction/RemoteControlProtocolMIDI/Source/RemoteControlProtocolMIDI/RemoteControlProtocolMIDI.Build.cs
--- a/Engine/Plugins/VirtualProduction/RemoteControlProtocolMIDI/Source/RemoteControlProtocolMIDI/RemoteControlProtocolMIDI.Build.cs
+++ b/Engine/Plugins/VirtualProduction/RemoteControlProtocolMIDI/Source/RemoteControlProtocolMIDI/RemoteControlProtocolMIDI.Build.cs
@@ -20,5 +20,14 @@ public class RemoteControlProtocolMIDI : ModuleRules
 				"RemoteControlProtocol",
 			}
 		);
+		
+		if (Target.bBuildEditor)
+		{
+			PrivateDependencyModuleNames.AddRange(
+				new string[] {
+					"RemoteControlProtocolWidgets"
+				}
+			);
+		}
 	}
 }
diff --git a/Engine/Plugins/VirtualProduction/RemoteControlProtocolOSC/Source/RemoteControlProtocolOSC/Private/RemoteControlProtocolOSC.cpp b/Engine/Plugins/VirtualProduction/RemoteControlProtocolOSC/Source/RemoteControlProtocolOSC/Private/RemoteControlProtocolOSC.cpp
--- a/Engine/Plugins/VirtualProduction/RemoteControlProtocolOSC/Source/RemoteControlProtocolOSC/Private/RemoteControlProtocolOSC.cpp
+++ b/Engine/Plugins/VirtualProduction/RemoteControlProtocolOSC/Source/RemoteControlProtocolOSC/Private/RemoteControlProtocolOSC.cpp
@@ -4,6 +4,11 @@
 
 #include "RemoteControlLogger.h"
 
+#if WITH_EDITOR
+#include "IRCProtocolBindingList.h"
+#include "IRemoteControlProtocolWidgetsModule.h"
+#endif
+
 #include "OSCManager.h"
 #include "OSCMessage.h"
 
@@ -41,6 +46,11 @@ void FRemoteControlProtocolOSC::Unbind(FRemoteControlProtocolEntityPtr InRemoteC
 void FRemoteControlProtocolOSC::OSCReceivedMessageEvent(const FOSCMessage& Message, const FString& IPAddress, uint16 Port)
 {	
 	const FOSCAddress& Address = Message.GetAddress();
+	
+#if WITH_EDITOR
+	ProcessAutoBinding(Address);
+#endif
+	
 	TArray<float> FloatValues;
 	UOSCManager::GetAllFloats(Message, FloatValues);
 
@@ -86,6 +96,38 @@ void FRemoteControlProtocolOSC::OSCReceivedMessageEvent(const FOSCMessage& Messa
 	}
 }
 
+#if WITH_EDITOR
+void FRemoteControlProtocolOSC::ProcessAutoBinding(const FOSCAddress& InAddress)
+{
+	// Bind only in Editor
+	if (!GIsEditor)
+	{
+		return;
+	}
+	
+	IRemoteControlProtocolWidgetsModule& RCWidgetsModule = IRemoteControlProtocolWidgetsModule::Get();
+	const TSharedPtr<IRCProtocolBindingList> RCProtocolBindingList = RCWidgetsModule.GetProtocolBindingList();
+	if (RCProtocolBindingList.IsValid())
+	{
+		for (const TSharedPtr<TStructOnScope<FRemoteControlProtocolEntity>>& ProtocolEntityPtr : RCProtocolBindingList->GetAwaitingProtocolEntities())
+		{
+			if (ProtocolEntityPtr.IsValid())
+			{
+				if ((*ProtocolEntityPtr)->GetBindingStatus() == ERCBindingStatus::Awaiting)
+				{
+					Unbind(ProtocolEntityPtr);
+
+					FRemoteControlOSCProtocolEntity* OSCProtocolEntity = ProtocolEntityPtr->CastChecked<FRemoteControlOSCProtocolEntity>();
+					OSCProtocolEntity->PathName = *InAddress.GetFullPath();
+
+					Bind(ProtocolEntityPtr);
+				}
+			}	
+		}
+	}
+}
+#endif
+
 void FRemoteControlProtocolOSC::UnbindAll()
 {
 	Bindings.Empty();
diff --git a/Engine/Plugins/VirtualProduction/RemoteControlProtocolOSC/Source/RemoteControlProtocolOSC/Private/RemoteControlProtocolOSC.h b/Engine/Plugins/VirtualProduction/RemoteControlProtocolOSC/Source/RemoteControlProtocolOSC/Private/RemoteControlProtocolOSC.h
--- a/Engine/Plugins/VirtualProduction/RemoteControlProtocolOSC/Source/RemoteControlProtocolOSC/Private/RemoteControlProtocolOSC.h
+++ b/Engine/Plugins/VirtualProduction/RemoteControlProtocolOSC/Source/RemoteControlProtocolOSC/Private/RemoteControlProtocolOSC.h
@@ -7,6 +7,7 @@
 
 #include "RemoteControlProtocolOSC.generated.h"
 
+struct FOSCAddress;
 struct FOSCMessage;
 
 /**
@@ -49,9 +50,17 @@ public:
 	virtual UScriptStruct* GetProtocolScriptStruct() const override { return FRemoteControlOSCProtocolEntity::StaticStruct(); }
 	//~ End IRemoteControlProtocol interface
 
-	/** Recieve OSC server message handler */
+	/** Receive OSC server message handler */
 	void OSCReceivedMessageEvent(const FOSCMessage& Message, const FString& IPAddress, uint16 Port);
 
+#if WITH_EDITOR
+	/**
+	 * Process the AutoBinding to the Remote Control Entity
+	 * @param InAddress	OSC address structure
+	 */
+	void ProcessAutoBinding(const FOSCAddress& InAddress);
+#endif
+
 private:
 	/** Map of the OSC bindings */
 	TMap<FName, TArray<FRemoteControlProtocolEntityWeakPtr>> Bindings;
diff --git a/Engine/Plugins/VirtualProduction/RemoteControlProtocolOSC/Source/RemoteControlProtocolOSC/RemoteControlProtocolOSC.Build.cs b/Engine/Plugins/VirtualProduction/RemoteControlProtocolOSC/Source/RemoteControlProtocolOSC/RemoteControlProtocolOSC.Build.cs
--- a/Engine/Plugins/VirtualProduction/RemoteControlProtocolOSC/Source/RemoteControlProtocolOSC/RemoteControlProtocolOSC.Build.cs
+++ b/Engine/Plugins/VirtualProduction/RemoteControlProtocolOSC/Source/RemoteControlProtocolOSC/RemoteControlProtocolOSC.Build.cs
@@ -27,6 +27,7 @@ public class RemoteControlProtocolOSC : ModuleRules
 			PrivateDependencyModuleNames.AddRange(
 				new string[]
 				{
+				    "RemoteControlProtocolWidgets",
 					"Settings"
 				}
 			);
