Remote Control DMX - DMX ports are now available before interacting with ports in project settings

diff --git a/Engine/Plugins/VirtualProduction/RemoteControlProtocolDMX/Source/RemoteControlProtocolDMX/Private/RemoteControlProtocolDMX.cpp b/Engine/Plugins/VirtualProduction/RemoteControlProtocolDMX/Source/RemoteControlProtocolDMX/Private/RemoteControlProtocolDMX.cpp
--- a/Engine/Plugins/VirtualProduction/RemoteControlProtocolDMX/Source/RemoteControlProtocolDMX/Private/RemoteControlProtocolDMX.cpp
+++ b/Engine/Plugins/VirtualProduction/RemoteControlProtocolDMX/Source/RemoteControlProtocolDMX/Private/RemoteControlProtocolDMX.cpp
@@ -91,10 +91,10 @@ void FRemoteControlDMXProtocolEntity::UpdateInputPort()
 {
 	if (bUseDefaultInputPort || !InputPortId.IsValid() || !FDMXPortManager::Get().FindInputPortByGuid(InputPortId))
 	{
-		const URemoteControlProtocolDMXSettings* RemoteControlDMXSettings = GetDefault<URemoteControlProtocolDMXSettings>();
+		URemoteControlProtocolDMXSettings* RemoteControlDMXSettings = GetMutableDefault<URemoteControlProtocolDMXSettings>();
 		if (RemoteControlDMXSettings)
 		{
-			const FGuid& DefaultInputPortId = RemoteControlDMXSettings->DefaultInputPortId;
+			const FGuid& DefaultInputPortId = RemoteControlDMXSettings->GetOrCreateDefaultInputPortId();
 			FDMXInputPortSharedPtr DefaultInputPort = FDMXPortManager::Get().FindInputPortByGuid(DefaultInputPortId);
 			if (DefaultInputPort.IsValid())
 			{
diff --git a/Engine/Plugins/VirtualProduction/RemoteControlProtocolDMX/Source/RemoteControlProtocolDMX/Private/RemoteControlProtocolDMXSettings.cpp b/Engine/Plugins/VirtualProduction/RemoteControlProtocolDMX/Source/RemoteControlProtocolDMX/Private/RemoteControlProtocolDMXSettings.cpp
--- a/Engine/Plugins/VirtualProduction/RemoteControlProtocolDMX/Source/RemoteControlProtocolDMX/Private/RemoteControlProtocolDMXSettings.cpp
+++ b/Engine/Plugins/VirtualProduction/RemoteControlProtocolDMX/Source/RemoteControlProtocolDMX/Private/RemoteControlProtocolDMXSettings.cpp
@@ -7,31 +7,40 @@
 
 FSimpleMulticastDelegate URemoteControlProtocolDMXSettings::OnRemoteControlProtocolDMXSettingsChangedDelegate;
 
-void URemoteControlProtocolDMXSettings::PostLoad()
+#if WITH_EDITOR
+void URemoteControlProtocolDMXSettings::PostEditChangeProperty(struct FPropertyChangedEvent& PropertyChangedEvent)
 {
-	Super::PostLoad();
+	Super::PostEditChangeProperty(PropertyChangedEvent);
 
+	if (PropertyChangedEvent.ChangeType != EPropertyChangeType::Interactive)
+	{
+		OnRemoteControlProtocolDMXSettingsChangedDelegate.Broadcast();
+	}
+}
+#endif // WITH_EDITOR
+
+FGuid URemoteControlProtocolDMXSettings::GetOrCreateDefaultInputPortId()
+{
 	if (!DefaultInputPortId.IsValid() || !FDMXPortManager::Get().FindInputPortByGuid(DefaultInputPortId))
 	{
 		const TArray<FDMXInputPortSharedRef> InputPorts = FDMXPortManager::Get().GetInputPorts();
 		if (InputPorts.Num() > 0)
 		{
+#if WITH_EDITOR
+			Modify();
+			PreEditChange(URemoteControlProtocolDMXSettings::StaticClass()->FindPropertyByName(GET_MEMBER_NAME_CHECKED(URemoteControlProtocolDMXSettings, DefaultInputPortId)));
+#endif 
+
 			DefaultInputPortId = InputPorts[0]->GetPortGuid();
-		}
-	}
-}
 
 #if WITH_EDITOR
-void URemoteControlProtocolDMXSettings::PostEditChangeProperty(struct FPropertyChangedEvent& PropertyChangedEvent)
-{
-	Super::PostEditChangeProperty(PropertyChangedEvent);
-
-	if (PropertyChangedEvent.ChangeType != EPropertyChangeType::Interactive)
-	{
-		OnRemoteControlProtocolDMXSettingsChangedDelegate.Broadcast();
+			PostEditChange();
+#endif
+		}
 	}
+
+	return DefaultInputPortId;
 }
-#endif // WITH_EDITOR
 
 FSimpleMulticastDelegate& URemoteControlProtocolDMXSettings::GetOnRemoteControlProtocolDMXSettingsChanged()
 {
diff --git a/Engine/Plugins/VirtualProduction/RemoteControlProtocolDMX/Source/RemoteControlProtocolDMX/Public/RemoteControlProtocolDMXSettings.h b/Engine/Plugins/VirtualProduction/RemoteControlProtocolDMX/Source/RemoteControlProtocolDMX/Public/RemoteControlProtocolDMXSettings.h
--- a/Engine/Plugins/VirtualProduction/RemoteControlProtocolDMX/Source/RemoteControlProtocolDMX/Public/RemoteControlProtocolDMXSettings.h
+++ b/Engine/Plugins/VirtualProduction/RemoteControlProtocolDMX/Source/RemoteControlProtocolDMX/Public/RemoteControlProtocolDMXSettings.h
@@ -19,19 +19,23 @@ class REMOTECONTROLPROTOCOLDMX_API URemoteControlProtocolDMXSettings : public UO
 
 public:
 	//~ Begin UObject interface
-	virtual void PostLoad() override;
 #if WITH_EDITOR
 	virtual void PostEditChangeProperty(struct FPropertyChangedEvent& PropertyChangedEvent) override;
 #endif
 	//~ End UObject interface
 
-	/** DMX Default Device */
-	UPROPERTY(Config, EditAnywhere, Category = Mapping)
-	FGuid DefaultInputPortId;
+	UFUNCTION(BlueprintGetter)
+	FGuid GetOrCreateDefaultInputPortId();
 
 	/** Returns a delegate broadcast whenever the Remote Control Protocol DMX Settings changed */
 	static FSimpleMulticastDelegate& GetOnRemoteControlProtocolDMXSettingsChanged();
 
+	static FName GetDefaultInputPortIdPropertyNameChecked() { return GET_MEMBER_NAME_CHECKED(URemoteControlProtocolDMXSettings, DefaultInputPortId); }
+
 private:
+	/** DMX Default Device */
+	UPROPERTY(Config, EditAnywhere, BlueprintReadOnly, Meta = (BlueprintGetter = GetOrCreateDefaultInputPortId, AllowPrivateAccess = true), Category = Mapping)
+	FGuid DefaultInputPortId;
+
 	static FSimpleMulticastDelegate OnRemoteControlProtocolDMXSettingsChangedDelegate;
 };
diff --git a/Engine/Plugins/VirtualProduction/RemoteControlProtocolDMX/Source/RemoteControlProtocolDMXEditor/Private/RemoteControlDMXProtocolEntityExtraSettingCustomization.cpp b/Engine/Plugins/VirtualProduction/RemoteControlProtocolDMX/Source/RemoteControlProtocolDMXEditor/Private/RemoteControlDMXProtocolEntityExtraSettingCustomization.cpp
--- a/Engine/Plugins/VirtualProduction/RemoteControlProtocolDMX/Source/RemoteControlProtocolDMXEditor/Private/RemoteControlDMXProtocolEntityExtraSettingCustomization.cpp
+++ b/Engine/Plugins/VirtualProduction/RemoteControlProtocolDMX/Source/RemoteControlProtocolDMXEditor/Private/RemoteControlDMXProtocolEntityExtraSettingCustomization.cpp
@@ -42,7 +42,7 @@ void FRemoteControlDMXProtocolEntityExtraSettingCustomization::CustomizeChildren
 	}
 
 	// Handle project setting changes
-	const URemoteControlProtocolDMXSettings* ProtocolDMXSettings = GetDefault<URemoteControlProtocolDMXSettings>();
+	URemoteControlProtocolDMXSettings* ProtocolDMXSettings = GetMutableDefault<URemoteControlProtocolDMXSettings>();
 	ProtocolDMXSettings->GetOnRemoteControlProtocolDMXSettingsChanged().AddSP(this, &FRemoteControlDMXProtocolEntityExtraSettingCustomization::OnInputPortChanged);
 
 	// Handle property changes
@@ -84,7 +84,7 @@ void FRemoteControlDMXProtocolEntityExtraSettingCustomization::CustomizeChildren
 		FGuid Result = GetPortGuid();
 		if (!Result.IsValid())
 		{
-			Result = ProtocolDMXSettings->DefaultInputPortId;
+			Result = ProtocolDMXSettings->GetOrCreateDefaultInputPortId();
 		}
 
 		return Result;
@@ -131,8 +131,8 @@ void FRemoteControlDMXProtocolEntityExtraSettingCustomization::OnInputPortChange
 	{
 		if (bUseDefaultInputPort)
 		{
-			const URemoteControlProtocolDMXSettings* ProtocolDMXSettings = GetDefault<URemoteControlProtocolDMXSettings>();
-			FGuid PortGuid = ProtocolDMXSettings->DefaultInputPortId;
+			URemoteControlProtocolDMXSettings* ProtocolDMXSettings = GetMutableDefault<URemoteControlProtocolDMXSettings>();
+			FGuid PortGuid = ProtocolDMXSettings->GetOrCreateDefaultInputPortId();
 
 			SetPortGuid(PortGuid);
 		}
diff --git a/Engine/Plugins/VirtualProduction/RemoteControlProtocolDMX/Source/RemoteControlProtocolDMXEditor/Private/RemoteControlProtocolDMXSettingsDetails.cpp b/Engine/Plugins/VirtualProduction/RemoteControlProtocolDMX/Source/RemoteControlProtocolDMXEditor/Private/RemoteControlProtocolDMXSettingsDetails.cpp
--- a/Engine/Plugins/VirtualProduction/RemoteControlProtocolDMX/Source/RemoteControlProtocolDMXEditor/Private/RemoteControlProtocolDMXSettingsDetails.cpp
+++ b/Engine/Plugins/VirtualProduction/RemoteControlProtocolDMX/Source/RemoteControlProtocolDMXEditor/Private/RemoteControlProtocolDMXSettingsDetails.cpp
@@ -28,7 +28,7 @@ TSharedRef<IDetailCustomization> FRemoteControlProtocolDMXSettingsDetails::MakeI
 void FRemoteControlProtocolDMXSettingsDetails::CustomizeDetails(IDetailLayoutBuilder& DetailBuilder)
 {
 	// Customize the InputPortId property view
-	DefaultInputPortIdHandle = DetailBuilder.GetProperty(GET_MEMBER_NAME_CHECKED(URemoteControlProtocolDMXSettings, DefaultInputPortId));
+	DefaultInputPortIdHandle = DetailBuilder.GetProperty(URemoteControlProtocolDMXSettings::GetDefaultInputPortIdPropertyNameChecked());
 	if (!ensure(DefaultInputPortIdHandle.IsValid()))
 	{
 		return;
