Remote Control DMX - Fix a crash when deleting output ports from project settings

diff --git a/Engine/Plugins/VirtualProduction/RemoteControlProtocolDMX/Source/RemoteControlProtocolDMXEditor/Private/RemoteControlProtocolDMXSettingsDetails.cpp b/Engine/Plugins/VirtualProduction/RemoteControlProtocolDMX/Source/RemoteControlProtocolDMXEditor/Private/RemoteControlProtocolDMXSettingsDetails.cpp
--- a/Engine/Plugins/VirtualProduction/RemoteControlProtocolDMX/Source/RemoteControlProtocolDMXEditor/Private/RemoteControlProtocolDMXSettingsDetails.cpp
+++ b/Engine/Plugins/VirtualProduction/RemoteControlProtocolDMX/Source/RemoteControlProtocolDMXEditor/Private/RemoteControlProtocolDMXSettingsDetails.cpp
@@ -11,6 +11,7 @@
 #include "DetailWidgetRow.h"
 #include "EditorStyleSet.h"
 #include "IAssetRegistry.h"
+#include "IPropertyUtilities.h"
 #include "AssetRegistry/AssetRegistryModule.h"
 #include "IO/DMXInputPort.h"
 #include "IO/DMXPortManager.h"
@@ -27,6 +28,8 @@ TSharedRef<IDetailCustomization> FRemoteControlProtocolDMXSettingsDetails::MakeI
 
 void FRemoteControlProtocolDMXSettingsDetails::CustomizeDetails(IDetailLayoutBuilder& DetailBuilder)
 {
+	PropertyUtilities = DetailBuilder.GetPropertyUtilities();
+
 	// Customize the InputPortId property view
 	DefaultInputPortIdHandle = DetailBuilder.GetProperty(URemoteControlProtocolDMXSettings::GetDefaultInputPortIdPropertyNameChecked());
 	if (!ensure(DefaultInputPortIdHandle.IsValid()))
@@ -76,10 +79,16 @@ void FRemoteControlProtocolDMXSettingsDetails::CustomizeDetails(IDetailLayoutBui
 
 void FRemoteControlProtocolDMXSettingsDetails::OnPortSelected()
 {
-	FDMXInputPortSharedPtr InputPort = PortSelector->GetSelectedInputPort();
-	FGuid NewPortGuid = InputPort->GetPortGuid();
+	if (FDMXInputPortSharedPtr InputPort = PortSelector->GetSelectedInputPort())
+	{
+		const FGuid NewPortGuid = InputPort->GetPortGuid();
 
-	SetPortGuid(NewPortGuid);
+		SetPortGuid(NewPortGuid);
+	}
+	else
+	{
+		PropertyUtilities->ForceRefresh();
+	}
 }
 
 FGuid FRemoteControlProtocolDMXSettingsDetails::GetPortGuid() const
diff --git a/Engine/Plugins/VirtualProduction/RemoteControlProtocolDMX/Source/RemoteControlProtocolDMXEditor/Private/RemoteControlProtocolDMXSettingsDetails.h b/Engine/Plugins/VirtualProduction/RemoteControlProtocolDMX/Source/RemoteControlProtocolDMXEditor/Private/RemoteControlProtocolDMXSettingsDetails.h
--- a/Engine/Plugins/VirtualProduction/RemoteControlProtocolDMX/Source/RemoteControlProtocolDMXEditor/Private/RemoteControlProtocolDMXSettingsDetails.h
+++ b/Engine/Plugins/VirtualProduction/RemoteControlProtocolDMX/Source/RemoteControlProtocolDMXEditor/Private/RemoteControlProtocolDMXSettingsDetails.h
@@ -5,6 +5,7 @@
 #include "IDetailCustomization.h"
 
 class IPropertyHandle;
+class IPropertyUtilities;
 class SDMXPortSelector;
 
 /**
@@ -37,4 +38,7 @@ private:
 
 	/** Widget to select a port */
 	TSharedPtr<SDMXPortSelector> PortSelector;
+
+	/** Property utilities for this customization */
+	TSharedPtr<IPropertyUtilities> PropertyUtilities;
 };
