RemoteControl DMX - Channel Property is now visible and editable in Remote Control bindings as expected

diff --git a/Engine/Plugins/VirtualProduction/RemoteControlProtocolDMX/Source/RemoteControlProtocolDMX/Private/RemoteControlProtocolDMX.cpp b/Engine/Plugins/VirtualProduction/RemoteControlProtocolDMX/Source/RemoteControlProtocolDMX/Private/RemoteControlProtocolDMX.cpp
--- a/Engine/Plugins/VirtualProduction/RemoteControlProtocolDMX/Source/RemoteControlProtocolDMX/Private/RemoteControlProtocolDMX.cpp
+++ b/Engine/Plugins/VirtualProduction/RemoteControlProtocolDMX/Source/RemoteControlProtocolDMX/Private/RemoteControlProtocolDMX.cpp
@@ -6,6 +6,7 @@
 #include "DMXProtocolSettings.h"
 #include "DMXProtocolTypes.h"
 #include "RemoteControlLogger.h"
+#include "RemoteControlProtocolDMXObjectVersion.h"
 #include "RemoteControlProtocolDMXSettings.h"
 #include "Interfaces/IDMXProtocol.h"
 #include "IO/DMXInputPort.h"
@@ -23,7 +24,7 @@ const FName FRemoteControlProtocolDMX::ProtocolName = TEXT("DMX");
 
 uint8 FRemoteControlDMXProtocolEntity::GetRangePropertySize() const
 {
-	switch (DataType)
+	switch (ExtraSetting.DataType)
 	{
 		default:
 		case EDMXFixtureSignalFormat::E8Bit:
@@ -43,7 +44,7 @@ uint8 FRemoteControlDMXProtocolEntity::GetRangePropertySize() const
 
 const FString& FRemoteControlDMXProtocolEntity::GetRangePropertyMaxValue() const
 {
-	switch (DataType)
+	switch (ExtraSetting.DataType)
 	{
 		default:
 		case EDMXFixtureSignalFormat::E8Bit:
@@ -89,7 +90,7 @@ void FRemoteControlDMXProtocolEntity::Initialize()
 
 void FRemoteControlDMXProtocolEntity::UpdateInputPort()
 {
-	if (bUseDefaultInputPort || !InputPortId.IsValid() || !FDMXPortManager::Get().FindInputPortByGuid(InputPortId))
+	if (ExtraSetting.bUseDefaultInputPort || !ExtraSetting.InputPortId.IsValid() || !FDMXPortManager::Get().FindInputPortByGuid(ExtraSetting.InputPortId))
 	{
 		URemoteControlProtocolDMXSettings* RemoteControlDMXSettings = GetMutableDefault<URemoteControlProtocolDMXSettings>();
 		if (RemoteControlDMXSettings)
@@ -98,12 +99,38 @@ void FRemoteControlDMXProtocolEntity::UpdateInputPort()
 			FDMXInputPortSharedPtr DefaultInputPort = FDMXPortManager::Get().FindInputPortByGuid(DefaultInputPortId);
 			if (DefaultInputPort.IsValid())
 			{
-				InputPortId = DefaultInputPortId;
+				ExtraSetting.InputPortId = DefaultInputPortId;
 			}
 		}
 	}
 }
 
+bool FRemoteControlDMXProtocolEntity::Serialize(FArchive& Ar)
+{
+	Ar.UsingCustomVersion(FRemoteControlProtocolDMXObjectVersion::GUID);
+	
+	// Don't actually serialize, just write the custom version for PostSerialize
+	return false;
+}
+
+void FRemoteControlDMXProtocolEntity::PostSerialize(const FArchive& Ar)
+{
+	if (Ar.IsLoading())
+	{
+		if (Ar.CustomVer(FRemoteControlProtocolDMXObjectVersion::GUID) < FRemoteControlProtocolDMXObjectVersion::MoveRemoteControlProtocolDMXEntityPropertiesToExtraSettingStruct)
+		{
+			// Move relevant properties that were moved to the the ExtraSetting member in 5.0 to the ExtraSetting member, so they can be customized. 
+			PRAGMA_DISABLE_DEPRECATION_WARNINGS
+			ExtraSetting.bUseDefaultInputPort = bUseDefaultInputPort_DEPRECATED;
+			ExtraSetting.bUseLSB = bUseLSB_DEPRECATED;
+			ExtraSetting.DataType = DataType_DEPRECATED;
+			ExtraSetting.InputPortId = InputPortId_DEPRECATED;
+			ExtraSetting.Universe = Universe_DEPRECATED;
+			PRAGMA_ENABLE_DEPRECATION_WARNINGS
+		}
+	}
+}
+
 void FRemoteControlProtocolDMX::Bind(FRemoteControlProtocolEntityPtr InRemoteControlProtocolEntityPtr)
 {
 	if (!ensure(InRemoteControlProtocolEntityPtr.IsValid()))
@@ -119,7 +146,7 @@ void FRemoteControlProtocolDMX::Bind(FRemoteControlProtocolEntityPtr InRemoteCon
 			{
 				const FRemoteControlDMXProtocolEntity* ComparedDMXProtocolEntity = ProtocolEntity->CastChecked<FRemoteControlDMXProtocolEntity>();
 
-				if (ComparedDMXProtocolEntity->Universe == DMXProtocolEntity->Universe &&
+				if (ComparedDMXProtocolEntity->ExtraSetting.Universe == DMXProtocolEntity->ExtraSetting.Universe &&
 					ComparedDMXProtocolEntity->ExtraSetting.StartingChannel == DMXProtocolEntity->ExtraSetting.StartingChannel &&
 					ComparedDMXProtocolEntity->GetPropertyId() == DMXProtocolEntity->GetPropertyId())
 				{
@@ -164,7 +191,7 @@ void FRemoteControlProtocolDMX::OnEndFrame()
 		{
 			FRemoteControlDMXProtocolEntity* DMXProtocolEntity = ProtocolEntity->CastChecked<FRemoteControlDMXProtocolEntity>();
 
-			const FGuid& PortId =  DMXProtocolEntity->InputPortId;				
+			const FGuid& PortId =  DMXProtocolEntity->ExtraSetting.InputPortId;				
 
 			const FDMXInputPortSharedRef* InputPortPtr = InputPorts.FindByPredicate([&PortId](const FDMXInputPortSharedRef& InputPort) {
                 return InputPort->GetPortGuid() == PortId;
@@ -178,7 +205,7 @@ void FRemoteControlProtocolDMX::OnEndFrame()
 			const FDMXInputPortSharedRef& InputPort = *InputPortPtr;
 
 			// Get universe DMX signal
-			InputPort->GameThreadGetDMXSignal(DMXProtocolEntity->Universe, DMXProtocolEntity->LastSignalPtr);
+			InputPort->GameThreadGetDMXSignal(DMXProtocolEntity->ExtraSetting.Universe, DMXProtocolEntity->LastSignalPtr);
 			if (DMXProtocolEntity->LastSignalPtr.IsValid())
 			{
 				const FDMXSignalSharedPtr& LastSignalPtr = DMXProtocolEntity->LastSignalPtr;
@@ -206,12 +233,12 @@ void FRemoteControlProtocolDMX::ProcessAndApplyProtocolValue(const FDMXSignalSha
 	
 	FRemoteControlDMXProtocolEntity* DMXProtocolEntity = InProtocolEntityPtr->CastChecked<FRemoteControlDMXProtocolEntity>();
 	const uint8* ChannelData = &InSignal->ChannelData[InDMXOffset];	
-	const uint8 NumChannelsToOccupy = FDMXConversions::GetSizeOfSignalFormat(DMXProtocolEntity->DataType);
+	const uint8 NumChannelsToOccupy = FDMXConversions::GetSizeOfSignalFormat(DMXProtocolEntity->ExtraSetting.DataType);
 
 	if(DMXProtocolEntity->CacheDMXBuffer.Num() != NumChannelsToOccupy ||
 		FMemory::Memcmp(DMXProtocolEntity->CacheDMXBuffer.GetData(), ChannelData, NumChannelsToOccupy) != 0)
 	{
-		const uint32 DMXValue = UDMXEntityFixtureType::BytesToInt(DMXProtocolEntity->DataType, DMXProtocolEntity->bUseLSB, ChannelData);
+		const uint32 DMXValue = UDMXEntityFixtureType::BytesToInt(DMXProtocolEntity->ExtraSetting.DataType, DMXProtocolEntity->ExtraSetting.bUseLSB, ChannelData);
 		
 #if WITH_EDITOR
 		FRemoteControlLogger::Get().Log(ProtocolName, [&InSignal, DMXValue]
@@ -219,7 +246,7 @@ void FRemoteControlProtocolDMX::ProcessAndApplyProtocolValue(const FDMXSignalSha
 			return FText::Format(LOCTEXT("DMXEventLog","ExternUniverseID {0}, DMXValue {1}"), InSignal->ExternUniverseID, DMXValue);
 		});
 #endif
-	
+
 		QueueValue(InProtocolEntityPtr, DMXValue);
 
 		// update cached buffer
diff --git a/Engine/Plugins/VirtualProduction/RemoteControlProtocolDMX/Source/RemoteControlProtocolDMX/Private/RemoteControlProtocolDMXObjectVersion.cpp b/Engine/Plugins/VirtualProduction/RemoteControlProtocolDMX/Source/RemoteControlProtocolDMX/Private/RemoteControlProtocolDMXObjectVersion.cpp
new file mode 100644
--- /dev/null
+++ b/Engine/Plugins/VirtualProduction/RemoteControlProtocolDMX/Source/RemoteControlProtocolDMX/Private/RemoteControlProtocolDMXObjectVersion.cpp
@@ -0,0 +1,11 @@
+// Copyright Epic Games, Inc. All Rights Reserved.
+
+#include "RemoteControlProtocolDMXObjectVersion.h"
+
+#include "Serialization/CustomVersion.h"
+
+
+const FGuid FRemoteControlProtocolDMXObjectVersion::GUID(0x077636E9, 0x44119FD9, 0x26477F9D, 0x0A7E9F6E);
+
+// Register the custom version with core
+FCustomVersionRegistration GRegisterDMXRuntimeCustomVersion(FRemoteControlProtocolDMXObjectVersion::GUID, FRemoteControlProtocolDMXObjectVersion::LatestVersion, TEXT("RemoteControlProtocolDMXObjectVersion"));
diff --git a/Engine/Plugins/VirtualProduction/RemoteControlProtocolDMX/Source/RemoteControlProtocolDMX/Private/RemoteControlProtocolDMXObjectVersion.h b/Engine/Plugins/VirtualProduction/RemoteControlProtocolDMX/Source/RemoteControlProtocolDMX/Private/RemoteControlProtocolDMXObjectVersion.h
new file mode 100644
--- /dev/null
+++ b/Engine/Plugins/VirtualProduction/RemoteControlProtocolDMX/Source/RemoteControlProtocolDMX/Private/RemoteControlProtocolDMXObjectVersion.h
@@ -0,0 +1,28 @@
+// Copyright Epic Games, Inc. All Rights Reserved.
+#pragma once
+
+#include "CoreTypes.h"
+#include "Misc/Guid.h"
+
+// Custom serialization version for changes to Remote Control Protocol DMX Objects
+struct FRemoteControlProtocolDMXObjectVersion
+{
+	enum Type
+	{
+		// Roughly corresponds to 4.27
+		BeforeCustomVersionWasAdded = 0,
+
+		// Move Remote Control Protocol DMX Entity Properties to the Extra Setting struct so they can be customized
+		MoveRemoteControlProtocolDMXEntityPropertiesToExtraSettingStruct,
+
+		// -----<new versions can be added above this line>-------------------------------------------------
+		VersionPlusOne,
+		LatestVersion = VersionPlusOne - 1
+	};
+
+	// The GUID for this custom version number
+	const static FGuid GUID;
+
+private:
+	FRemoteControlProtocolDMXObjectVersion() {}
+};
diff --git a/Engine/Plugins/VirtualProduction/RemoteControlProtocolDMX/Source/RemoteControlProtocolDMX/Public/RemoteControlProtocolDMX.h b/Engine/Plugins/VirtualProduction/RemoteControlProtocolDMX/Source/RemoteControlProtocolDMX/Public/RemoteControlProtocolDMX.h
--- a/Engine/Plugins/VirtualProduction/RemoteControlProtocolDMX/Source/RemoteControlProtocolDMX/Public/RemoteControlProtocolDMX.h
+++ b/Engine/Plugins/VirtualProduction/RemoteControlProtocolDMX/Source/RemoteControlProtocolDMX/Public/RemoteControlProtocolDMX.h
@@ -22,9 +22,35 @@ struct FRemoteControlDMXProtocolEntityExtraSetting
 {
 	GENERATED_BODY();
 
-	/** Starting universe channel */
-	UPROPERTY(EditAnywhere, Category = Mapping, meta = (ClampMin = "1", ClampMax = "512", UIMin = "1", UIMax = "512"))
+
+	/** DMX universe id */
+	UPROPERTY(EditAnywhere, Category = Mapping, Meta = (ClampMin = "0", UIMin = "0"))
+	int32 Universe = 1;
+
+	/** Starting channel */
+	UPROPERTY(EditAnywhere, Category = Mapping, Meta = (ClampMin = "1", ClampMax = "512", UIMin = "1", UIMax = "512"))
 	int32 StartingChannel = 1;
+
+	/**
+	 * Least Significant Byte mode makes the individual bytes (channels) of the function be
+	 * interpreted with the first bytes being the lowest part of the number.
+	 * Most Fixtures use MSB (Most Significant Byte).
+	 */
+	UPROPERTY(EditAnywhere, Category = Mapping)
+	bool bUseLSB = false;
+
+	/** Defines the used number of channels (bytes) */
+	UPROPERTY(EditAnywhere, Category = Mapping)
+	EDMXFixtureSignalFormat DataType = EDMXFixtureSignalFormat::E8Bit;
+
+	/** If set to true, uses the default input port set in Remote Control Protocol project settings */
+	UPROPERTY(EditAnywhere, Category = Mapping)
+	bool bUseDefaultInputPort = true;
+
+	/** Reference of an input DMX port id */
+	UPROPERTY(EditAnywhere, Category = Mapping)
+	FGuid InputPortId;
+
 };
 
 /**
@@ -50,45 +76,53 @@ public:
 	/** Try to get the port ID from dmx protocol settings */
 	void UpdateInputPort();
 
-public:
-	/** DMX universe id */
-	UPROPERTY(EditAnywhere, Category = Mapping)
-	int32 Universe = 1;
-
-	/**
-	 * Least Significant Byte mode makes the individual bytes (channels) of the function be
-	 * interpreted with the first bytes being the lowest part of the number.
-	 * Most Fixtures use MSB (Most Significant Byte).
-	 */
-	UPROPERTY(EditAnywhere, Category = Mapping)
-	bool bUseLSB = false;
-
-	/** Defines the used number of channels (bytes) */
-	UPROPERTY(EditAnywhere, Category = Mapping)
-	EDMXFixtureSignalFormat DataType = EDMXFixtureSignalFormat::E8Bit;
-
 	/** Extra protocol settings. Primary using for customization */
-	UPROPERTY(EditAnywhere, Category = Mapping, meta=(ShowOnlyInnerProperties))
+	UPROPERTY(EditAnywhere, Category = Mapping, meta = (ShowOnlyInnerProperties))
 	FRemoteControlDMXProtocolEntityExtraSetting ExtraSetting;
 
 	/** DMX range input property template, used for binding. */
 	UPROPERTY(Transient)
 	uint32 RangeInputTemplate = 0;
 
-	/** If set to true, uses the default input port set in Remote Control Protocol project settings */
-	UPROPERTY(EditAnywhere, Category = Mapping)
-	bool bUseDefaultInputPort = true;
-
-	/** Reference of an input DMX port id */
-	UPROPERTY(EditAnywhere, Category = Mapping)
-	FGuid InputPortId;
-
 private:
 	/** DMX entity cache buffer. From 1 up to 4 channels, based on DataType */
 	TArray<uint8> CacheDMXBuffer;
 
 	/** A single, generic DMX signal. One universe of raw DMX data received */
 	FDMXSignalSharedPtr LastSignalPtr;
+
+public:
+	/** Called when the struct is serialized */
+	bool Serialize(FArchive& Ar);
+
+	/** Called after the struct is serialized */
+	void PostSerialize(const FArchive& Ar);
+
+	// DEPRECATED MEMBERS
+	// Deprecated 5.0
+	UPROPERTY(Meta = (DeprecatedProperty, DeprecationMessage = "This Property is deprecated and will be removed in a future release. It was moved to the ExtraSetting struct member so the property can be customized."))
+	int32 Universe_DEPRECATED = 1;
+	
+	UPROPERTY(Meta = (DeprecatedProperty, DeprecationMessage = "This Property is deprecated and will be removed in a future release. It was moved to the ExtraSetting struct member so the property can be customized."))
+	bool bUseLSB_DEPRECATED = false;
+
+	UPROPERTY(Meta = (DeprecatedProperty, DeprecationMessage = "This Property is deprecated and will be removed in a future release. It was moved to the ExtraSetting struct member so the property can be customized."))
+	EDMXFixtureSignalFormat DataType_DEPRECATED = EDMXFixtureSignalFormat::E8Bit;
+
+	UPROPERTY(Meta = (DeprecatedProperty, DeprecationMessage = "This Property is deprecated and will be removed in a future release. It was moved to the ExtraSetting struct member so the property can be customized."))
+	bool bUseDefaultInputPort_DEPRECATED = true;
+	
+	UPROPERTY(Meta = (DeprecatedProperty, DeprecationMessage = "This Property is deprecated and will be removed in a future release. It was moved to the ExtraSetting struct member so the property can be customized."))
+	FGuid InputPortId_DEPRECATED;
+};
+template<>
+struct TStructOpsTypeTraits<FRemoteControlDMXProtocolEntity> : public TStructOpsTypeTraitsBase2<FRemoteControlDMXProtocolEntity>
+{
+	enum
+	{
+		WithSerializer = true,
+		WithPostSerialize = true
+	};
 };
 
 /**
diff --git a/Engine/Plugins/VirtualProduction/RemoteControlProtocolDMX/Source/RemoteControlProtocolDMXEditor/Private/RemoteControlDMXProtocolEntityExtraSettingCustomization.cpp b/Engine/Plugins/VirtualProduction/RemoteControlProtocolDMX/Source/RemoteControlProtocolDMXEditor/Private/RemoteControlDMXProtocolEntityExtraSettingCustomization.cpp
--- a/Engine/Plugins/VirtualProduction/RemoteControlProtocolDMX/Source/RemoteControlProtocolDMXEditor/Private/RemoteControlDMXProtocolEntityExtraSettingCustomization.cpp
+++ b/Engine/Plugins/VirtualProduction/RemoteControlProtocolDMX/Source/RemoteControlProtocolDMXEditor/Private/RemoteControlDMXProtocolEntityExtraSettingCustomization.cpp
@@ -24,61 +24,62 @@ TSharedRef<IPropertyTypeCustomization> FRemoteControlDMXProtocolEntityExtraSetti
 	return MakeShared<FRemoteControlDMXProtocolEntityExtraSettingCustomization>();
 }
 
-void FRemoteControlDMXProtocolEntityExtraSettingCustomization::CustomizeHeader(TSharedRef<IPropertyHandle> InStructProperty,
-                                                                       FDetailWidgetRow& HeaderRow, IPropertyTypeCustomizationUtils& CustomizationUtils)
+void FRemoteControlDMXProtocolEntityExtraSettingCustomization::CustomizeHeader(TSharedRef<IPropertyHandle> InStructProperty, FDetailWidgetRow& HeaderRow, IPropertyTypeCustomizationUtils& CustomizationUtils)
 {
-	// No need to have header widget, only structure child widget should be present
+	// No need for a header widget, only structure child widgets should be present
 }
 
-void FRemoteControlDMXProtocolEntityExtraSettingCustomization::CustomizeChildren(const TSharedRef<IPropertyHandle> InStructProperty,
-                                                                         IDetailChildrenBuilder& ChildBuilder, IPropertyTypeCustomizationUtils& CustomizationUtils)
+void FRemoteControlDMXProtocolEntityExtraSettingCustomization::CustomizeChildren(const TSharedRef<IPropertyHandle> InStructProperty, IDetailChildrenBuilder& ChildBuilder, IPropertyTypeCustomizationUtils& CustomizationUtils)
 {
 	PropertyUtilities = CustomizationUtils.GetPropertyUtilities();
 
-	const TSharedPtr<IPropertyHandle> ProtocolEntityHandle = InStructProperty->GetParentHandle();
-	if (!ensure(ProtocolEntityHandle.IsValid()))
-	{
-		return;
-	}
+	// Remember relevant property handles
+	InputPortIdHandle = InStructProperty->GetChildHandle(GET_MEMBER_NAME_CHECKED(FRemoteControlDMXProtocolEntityExtraSetting, InputPortId));
+	UseDefaultInputPortHandle = InStructProperty->GetChildHandle(GET_MEMBER_NAME_CHECKED(FRemoteControlDMXProtocolEntityExtraSetting, bUseDefaultInputPort));
+	FixtureSignalFormatHandle = InStructProperty->GetChildHandle(GET_MEMBER_NAME_CHECKED(FRemoteControlDMXProtocolEntityExtraSetting, DataType));
+	StartingChannelPropertyHandle = InStructProperty->GetChildHandle(GET_MEMBER_NAME_CHECKED(FRemoteControlDMXProtocolEntityExtraSetting, StartingChannel));
 
 	// Handle project setting changes
 	URemoteControlProtocolDMXSettings* ProtocolDMXSettings = GetMutableDefault<URemoteControlProtocolDMXSettings>();
 	ProtocolDMXSettings->GetOnRemoteControlProtocolDMXSettingsChanged().AddSP(this, &FRemoteControlDMXProtocolEntityExtraSettingCustomization::OnInputPortChanged);
 
 	// Handle property changes
-	UseDefaultInputPortHandle = ProtocolEntityHandle->GetChildHandle(GET_MEMBER_NAME_CHECKED(FRemoteControlDMXProtocolEntity, bUseDefaultInputPort));
 	if (!ensure(UseDefaultInputPortHandle.IsValid()))
 	{
 		return;
 	}
 	UseDefaultInputPortHandle->SetOnPropertyValueChanged(FSimpleDelegate::CreateSP(this, &FRemoteControlDMXProtocolEntityExtraSettingCustomization::OnInputPortChanged));
-	UseDefaultInputPortHandle->MarkHiddenByCustomization();
 
-	FixtureSignalFormatHandle = ProtocolEntityHandle->GetChildHandle(GET_MEMBER_NAME_CHECKED(FRemoteControlDMXProtocolEntity, DataType));
 	if (!ensure(FixtureSignalFormatHandle.IsValid()))
 	{
 		return;
 	}
 	FixtureSignalFormatHandle->SetOnPropertyValueChanged(FSimpleDelegate::CreateSP(this, &FRemoteControlDMXProtocolEntityExtraSettingCustomization::OnFixtureSignalFormatChange));
 
-	StartingChannelPropertyHandle = ProtocolEntityHandle->GetChildHandle(GET_MEMBER_NAME_CHECKED(FRemoteControlDMXProtocolEntityExtraSetting, StartingChannel));
 	if (!ensure(StartingChannelPropertyHandle.IsValid()))
 	{
 		return;
 	}
 	StartingChannelPropertyHandle->SetOnPropertyValueChanged(FSimpleDelegate::CreateSP(this, &FRemoteControlDMXProtocolEntityExtraSettingCustomization::OnStartingChannelChange));
 
-	// Add the bUseDefaultInputPort property before the Input Port Selector
-	ChildBuilder.AddProperty(UseDefaultInputPortHandle.ToSharedRef());
-
-	// Customize the InputPortId property view
-	InputPortIdHandle = ProtocolEntityHandle->GetChildHandle(GET_MEMBER_NAME_CHECKED(FRemoteControlDMXProtocolEntity, InputPortId));
-	if (!ensure(InputPortIdHandle.IsValid()))
+	// Add all properties besides the InputPortId
+	uint32 NumChildren = 0;
+	InStructProperty->GetNumChildren(NumChildren);
+	for (uint32 ChildIndex = 0; ChildIndex < NumChildren; ChildIndex++)
 	{
-		return;
+		const TSharedPtr<IPropertyHandle> ChildHandle = InStructProperty->GetChildHandle(ChildIndex);
+		const FName PropertyName = ChildHandle->GetProperty() ? ChildHandle->GetProperty()->GetFName() : NAME_None;
+
+		if (PropertyName == GET_MEMBER_NAME_CHECKED(FRemoteControlDMXProtocolEntityExtraSetting, InputPortId))
+		{
+			continue;
+		}
+
+		ChildBuilder.AddProperty(ChildHandle.ToSharedRef());
+
 	}
-	InputPortIdHandle->MarkHiddenByCustomization();
 
+	// Show a port selector instead of the InputPortId
 	const FGuid PortGuid = [ProtocolDMXSettings, this]()
 	{
 		FGuid Result = GetPortGuid();
