Fix CIS Compile QAVirtualProduction Win64

#rnx

diff --git a/Engine/Plugins/VirtualProduction/RemoteControlProtocolDMX/Source/RemoteControlProtocolDMX/Private/RemoteControlProtocolDMXObjectVersion.cpp b/Engine/Plugins/VirtualProduction/RemoteControlProtocolDMX/Source/RemoteControlProtocolDMX/Private/RemoteControlProtocolDMXObjectVersion.cpp
--- a/Engine/Plugins/VirtualProduction/RemoteControlProtocolDMX/Source/RemoteControlProtocolDMX/Private/RemoteControlProtocolDMXObjectVersion.cpp
+++ b/Engine/Plugins/VirtualProduction/RemoteControlProtocolDMX/Source/RemoteControlProtocolDMX/Private/RemoteControlProtocolDMXObjectVersion.cpp
@@ -8,4 +8,4 @@
 const FGuid FRemoteControlProtocolDMXObjectVersion::GUID(0x077636E9, 0x44119FD9, 0x26477F9D, 0x0A7E9F6E);
 
 // Register the custom version with core
-FCustomVersionRegistration GRegisterDMXRuntimeCustomVersion(FRemoteControlProtocolDMXObjectVersion::GUID, FRemoteControlProtocolDMXObjectVersion::LatestVersion, TEXT("RemoteControlProtocolDMXObjectVersion"));
+FCustomVersionRegistration GRegisterRemoteControlProtocolDMXCustomVersion(FRemoteControlProtocolDMXObjectVersion::GUID, FRemoteControlProtocolDMXObjectVersion::LatestVersion, TEXT("RemoteControlProtocolDMXObjectVersion"));
