DMX: Fix incompatible module names preventing modules from being loaded correctly at launch.

#rnx

diff --git a/Engine/Plugins/VirtualProduction/DMX/DMXGDTF/Source/DMXGDTFTests/Private/DMXGDTFTestsModule.cpp b/Engine/Plugins/VirtualProduction/DMX/DMXGDTF/Source/DMXGDTFTests/Private/DMXGDTFTestsModule.cpp
--- a/Engine/Plugins/VirtualProduction/DMX/DMXGDTF/Source/DMXGDTFTests/Private/DMXGDTFTestsModule.cpp
+++ b/Engine/Plugins/VirtualProduction/DMX/DMXGDTF/Source/DMXGDTFTests/Private/DMXGDTFTestsModule.cpp
@@ -2,4 +2,4 @@
 
 #include "Modules/ModuleManager.h"
 
-IMPLEMENT_MODULE(FDefaultModuleImpl, DMXGDTFTestsModule)
+IMPLEMENT_MODULE(FDefaultModuleImpl, DMXGDTFTests)
diff --git a/Engine/Plugins/VirtualProduction/RemoteControlProtocolDMX/Source/RemoteControlProtocolDMXEditor/Private/RemoteControlProtocolDMXEditorModule.cpp b/Engine/Plugins/VirtualProduction/RemoteControlProtocolDMX/Source/RemoteControlProtocolDMXEditor/Private/RemoteControlProtocolDMXEditorModule.cpp
--- a/Engine/Plugins/VirtualProduction/RemoteControlProtocolDMX/Source/RemoteControlProtocolDMXEditor/Private/RemoteControlProtocolDMXEditorModule.cpp
+++ b/Engine/Plugins/VirtualProduction/RemoteControlProtocolDMX/Source/RemoteControlProtocolDMXEditor/Private/RemoteControlProtocolDMXEditorModule.cpp
@@ -44,4 +44,4 @@ void FRemoteControlProtocolDMXEditorModule::ShutdownModule()
 }
 
 
-IMPLEMENT_MODULE(FRemoteControlProtocolDMXEditorModule, RemoteControlProtocolDMXEditorModule);
+IMPLEMENT_MODULE(FRemoteControlProtocolDMXEditorModule, RemoteControlProtocolDMXEditor);
