Remote Control DMX: Remote Control now uses a DMX Library to handle DMX, Remote Control Presets can now be exported as MVR file

diff --git a/Engine/Plugins/VirtualProduction/RemoteControl/Source/RemoteControl/Public/UI/RCFieldGroupOrder.h b/Engine/Plugins/VirtualProduction/RemoteControl/Source/RemoteControl/Public/UI/RCFieldGroupOrder.h
deleted file mode 100644
diff --git a/Engine/Plugins/VirtualProduction/RemoteControl/Source/RemoteControl/Public/UI/RCFieldGroupType.h b/Engine/Plugins/VirtualProduction/RemoteControl/Source/RemoteControl/Public/UI/RCFieldGroupType.h
deleted file mode 100644
diff --git a/Engine/Plugins/VirtualProduction/RemoteControl/Source/RemoteControl/Public/UI/RCPanelExposedEntitiesListSettings.h b/Engine/Plugins/VirtualProduction/RemoteControl/Source/RemoteControl/Public/UI/RCPanelExposedEntitiesListSettings.h
deleted file mode 100644
diff --git a/Engine/Plugins/VirtualProduction/RemoteControl/Source/RemoteControlProtocol/Private/Signature/RCSignatureProtocolAction.cpp b/Engine/Plugins/VirtualProduction/RemoteControl/Source/RemoteControlProtocol/Private/Signature/RCSignatureProtocolAction.cpp
--- a/Engine/Plugins/VirtualProduction/RemoteControl/Source/RemoteControlProtocol/Private/Signature/RCSignatureProtocolAction.cpp
+++ b/Engine/Plugins/VirtualProduction/RemoteControl/Source/RemoteControlProtocol/Private/Signature/RCSignatureProtocolAction.cpp
@@ -360,6 +360,7 @@ void FRCSignatureProtocolAction::CreateProtocolEntity(const FRCSignatureActionCo
 		check(ProtocolEntity.GetScriptStruct() == ScriptStruct);
 		ScriptStruct->CopyScriptStruct(NewEntity->Get(), ProtocolEntity.GetPtr());
 	}
+	NewEntity->CastChecked<FRemoteControlProtocolEntity>()->Init(InContext.Preset, EntityId);
 
 	FRemoteControlProtocolBinding ProtocolBinding(ProtocolName, EntityId, NewEntity, FGuid::NewGuid());
 	AddMappings(InContext, InProtocol, /*InOut*/ProtocolBinding);
diff --git a/Engine/Plugins/VirtualProduction/RemoteControl/Source/RemoteControlProtocolWidgets/Private/ViewModels/ProtocolEntityViewModel.cpp b/Engine/Plugins/VirtualProduction/RemoteControl/Source/RemoteControlProtocolWidgets/Private/ViewModels/ProtocolEntityViewModel.cpp
--- a/Engine/Plugins/VirtualProduction/RemoteControl/Source/RemoteControlProtocolWidgets/Private/ViewModels/ProtocolEntityViewModel.cpp
+++ b/Engine/Plugins/VirtualProduction/RemoteControl/Source/RemoteControlProtocolWidgets/Private/ViewModels/ProtocolEntityViewModel.cpp
@@ -380,6 +380,11 @@ bool FProtocolEntityViewModel::IsValid() const
 			&& PropertyId.IsValid();
 }
 
+URemoteControlPreset* FProtocolEntityViewModel::GetPreset() const
+{
+	return Preset.Get();
+}
+
 void FProtocolEntityViewModel::PostUndo(bool bSuccess)
 {
 	check(IsValid());
diff --git a/Engine/Plugins/VirtualProduction/RemoteControl/Source/RemoteControlProtocolWidgets/Private/ViewModels/ProtocolEntityViewModel.h b/Engine/Plugins/VirtualProduction/RemoteControl/Source/RemoteControlProtocolWidgets/Private/ViewModels/ProtocolEntityViewModel.h
--- a/Engine/Plugins/VirtualProduction/RemoteControl/Source/RemoteControlProtocolWidgets/Private/ViewModels/ProtocolEntityViewModel.h
+++ b/Engine/Plugins/VirtualProduction/RemoteControl/Source/RemoteControlProtocolWidgets/Private/ViewModels/ProtocolEntityViewModel.h
@@ -68,6 +68,9 @@ public:
 	/** Checks validity of this ViewModel */
 	REMOTECONTROLPROTOCOLWIDGETS_API bool IsValid() const;
 
+	/** Returns the preset of this view model or nullptr if there is no outer preset */
+	URemoteControlPreset* GetPreset() const;
+
 public:
 	DECLARE_MULTICAST_DELEGATE_OneParam(FOnBindingAdded, TSharedRef<FProtocolBindingViewModel> /* InBindingViewModel */);
 	DECLARE_MULTICAST_DELEGATE_OneParam(FOnBindingRemoved, FGuid /* InBindingId */);
diff --git a/Engine/Plugins/VirtualProduction/RemoteControl/Source/RemoteControlProtocolWidgets/Private/Widgets/SRCProtocolBindingList.cpp b/Engine/Plugins/VirtualProduction/RemoteControl/Source/RemoteControlProtocolWidgets/Private/Widgets/SRCProtocolBindingList.cpp
--- a/Engine/Plugins/VirtualProduction/RemoteControl/Source/RemoteControlProtocolWidgets/Private/Widgets/SRCProtocolBindingList.cpp
+++ b/Engine/Plugins/VirtualProduction/RemoteControl/Source/RemoteControlProtocolWidgets/Private/Widgets/SRCProtocolBindingList.cpp
@@ -183,6 +183,11 @@ SRCProtocolBindingList::~SRCProtocolBindingList()
 	AwaitingProtocolEntities.Empty();
 }
 
+URemoteControlPreset* SRCProtocolBindingList::GetPreset() const
+{
+	return ViewModel.IsValid() ? ViewModel->GetPreset() : nullptr;
+}
+
 void SRCProtocolBindingList::AddProtocolBinding(const FName InProtocolName)
 {
 	if (!ViewModel.IsValid())
diff --git a/Engine/Plugins/VirtualProduction/RemoteControl/Source/RemoteControlProtocolWidgets/Private/Widgets/SRCProtocolBindingList.h b/Engine/Plugins/VirtualProduction/RemoteControl/Source/RemoteControlProtocolWidgets/Private/Widgets/SRCProtocolBindingList.h
--- a/Engine/Plugins/VirtualProduction/RemoteControl/Source/RemoteControlProtocolWidgets/Private/Widgets/SRCProtocolBindingList.h
+++ b/Engine/Plugins/VirtualProduction/RemoteControl/Source/RemoteControlProtocolWidgets/Private/Widgets/SRCProtocolBindingList.h
@@ -29,6 +29,7 @@ public:
 
 	//~ Begin IRCProtocolBindingList Interface
 	virtual void AddProtocolBinding(const FName InProtocolName) override;
+	virtual URemoteControlPreset* GetPreset() const;
 
 	/** Get the set of entities which is awaiting state and waiting for binding. */
 	virtual FRemoteControlProtocolEntitySet GetAwaitingProtocolEntities() const override
@@ -118,4 +119,7 @@ private:
 
 	/** Set of protocol entities with awaiting state and waiting for binding */
 	FRemoteControlProtocolEntitySet AwaitingProtocolEntities;
+
+	/** Map of protocol names and their support of awaiting protocol entities */
+	TMap<FName, bool> ProtocolNameToSupportsAwaitingProtocolEntitiesMap;
 };
diff --git a/Engine/Plugins/VirtualProduction/RemoteControl/Source/RemoteControlProtocolWidgets/Public/IRCProtocolBindingList.h b/Engine/Plugins/VirtualProduction/RemoteControl/Source/RemoteControlProtocolWidgets/Public/IRCProtocolBindingList.h
--- a/Engine/Plugins/VirtualProduction/RemoteControl/Source/RemoteControlProtocolWidgets/Public/IRCProtocolBindingList.h
+++ b/Engine/Plugins/VirtualProduction/RemoteControl/Source/RemoteControlProtocolWidgets/Public/IRCProtocolBindingList.h
@@ -5,6 +5,7 @@
 #include "CoreMinimal.h"
 #include "UObject/StructOnScope.h"
 
+class URemoteControlPreset;
 struct FRemoteControlProtocolEntity;
 
 /**
@@ -18,6 +19,9 @@ public:
 	/** Virtual destructor */
 	virtual ~IRCProtocolBindingList() = default;
 
+	/** Returns the remote control preset for which this list is currently drawn */
+	virtual URemoteControlPreset* GetPreset() const = 0;
+
 	/** Adds a new protocol binding of specified type to the current view model. */
 	virtual void AddProtocolBinding(const FName InProtocolName) = 0;
 
diff --git a/Engine/Plugins/VirtualProduction/RemoteControl/Source/RemoteControlUI/Private/UI/SRCPanelExposedEntitiesList.cpp b/Engine/Plugins/VirtualProduction/RemoteControl/Source/RemoteControlUI/Private/UI/SRCPanelExposedEntitiesList.cpp
--- a/Engine/Plugins/VirtualProduction/RemoteControl/Source/RemoteControlUI/Private/UI/SRCPanelExposedEntitiesList.cpp
+++ b/Engine/Plugins/VirtualProduction/RemoteControl/Source/RemoteControlUI/Private/UI/SRCPanelExposedEntitiesList.cpp
@@ -747,7 +747,7 @@ void SRCPanelExposedEntitiesList::RebuildListWithColumns(EEntitiesListMode InLis
 		if (bShouldAddColumns)
 		{
 			ColumnsToBeAdded = DefaultProtocolColumns.Union(ActiveProtocolColumns);
-		
+
 			for (const FName& ColumnToBeAdded : ColumnsToBeAdded)
 			{
 				InsertColumn(ColumnToBeAdded);
@@ -869,6 +869,11 @@ void SRCPanelExposedEntitiesList::StoreListSettings()
 		FRemoteControlUIModule::Get().GetExposedEntitiesListSettingsForProtocol(ActiveProtocol) :
 		nullptr;
 
+
+	const TOptional<FScopedTransaction> ChangeListSettingsTransaction = !GIsTransacting ?
+		FScopedTransaction(LOCTEXT("ChangeListSettingsTransaction", "Change Remote Control List Order")) :
+		TOptional<FScopedTransaction>{};
+
 	if (SettingsPtr)
 	{
 		FRCPanelExposedEntitiesListSettingsData Settings;
@@ -902,10 +907,8 @@ void SRCPanelExposedEntitiesList::RecallListSettings()
 		CurrentGroupSortType = DefaultSettings.FieldGroupOrder;
 	}
 
-	if (FieldsListView.IsValid())
-	{
-		OnCreateFieldGroup(CurrentGroupType);
-	}
+	CreateFieldGroup();
+	OrderGroups();
 }
 
 void SRCPanelExposedEntitiesList::OnObjectPropertyChange(UObject* InObject, FPropertyChangedEvent& InChangeEvent)
@@ -1591,7 +1594,8 @@ void SRCPanelExposedEntitiesList::CreateFieldGroup()
 		}
 	}
 
-	if (CurrentGroupType == ERCFieldGroupType::None)
+	if (CurrentGroupType == ERCFieldGroupType::None &&
+		FieldsListView.IsValid())
 	{
 		FieldEntities = CachedFieldEntities;
 		FieldsListView->RequestTreeRefresh();
diff --git a/Engine/Plugins/VirtualProduction/RemoteControl/Source/RemoteControlUI/Private/UI/SRCPanelExposedEntitiesList.h b/Engine/Plugins/VirtualProduction/RemoteControl/Source/RemoteControlUI/Private/UI/SRCPanelExposedEntitiesList.h
--- a/Engine/Plugins/VirtualProduction/RemoteControl/Source/RemoteControlUI/Private/UI/SRCPanelExposedEntitiesList.h
+++ b/Engine/Plugins/VirtualProduction/RemoteControl/Source/RemoteControlUI/Private/UI/SRCPanelExposedEntitiesList.h
@@ -10,7 +10,7 @@
 #include "UI/Filters/RCFilter.h"
 #include "UI/RCFieldGroupOrder.h"
 #include "UI/RCFieldGroupType.h"
-#include "UI/RCPanelExposedEntitiesListSettings.h"
+#include "UI/RCPanelExposedEntitiesListSettingsData.h"
 #include "UObject/StrongObjectPtr.h"
 #include "Widgets/DeclarativeSyntaxSupport.h"
 #include "Widgets/SCompoundWidget.h"
diff --git a/Engine/Plugins/VirtualProduction/RemoteControl/Source/RemoteControlUI/Public/UI/IRCPanelExposedEntitiesListSettingsForProtocol.h b/Engine/Plugins/VirtualProduction/RemoteControl/Source/RemoteControlUI/Public/UI/IRCPanelExposedEntitiesListSettingsForProtocol.h
--- a/Engine/Plugins/VirtualProduction/RemoteControl/Source/RemoteControlUI/Public/UI/IRCPanelExposedEntitiesListSettingsForProtocol.h
+++ b/Engine/Plugins/VirtualProduction/RemoteControl/Source/RemoteControlUI/Public/UI/IRCPanelExposedEntitiesListSettingsForProtocol.h
@@ -2,7 +2,7 @@
 
 #pragma once
 
-#include "UI/RCPanelExposedEntitiesListSettings.h"
+#include "UI/RCPanelExposedEntitiesListSettingsData.h"
 #include "UObject/NameTypes.h"
 
 class URemoteControlPreset;
diff --git a/Engine/Plugins/VirtualProduction/RemoteControl/Source/RemoteControlUI/Public/UI/RCFieldGroupOrder.h b/Engine/Plugins/VirtualProduction/RemoteControl/Source/RemoteControlUI/Public/UI/RCFieldGroupOrder.h
new file mode 100644
--- /dev/null
+++ b/Engine/Plugins/VirtualProduction/RemoteControl/Source/RemoteControlUI/Public/UI/RCFieldGroupOrder.h
@@ -0,0 +1,14 @@
+// Copyright Epic Games, Inc. All Rights Reserved.
+
+#pragma once
+
+#include "RCFieldGroupOrder.generated.h"
+
+/** Order of remote control field groups */
+UENUM()
+enum class ERCFieldGroupOrder : uint8
+{
+	None,
+	Ascending,
+	Descending
+};
diff --git a/Engine/Plugins/VirtualProduction/RemoteControl/Source/RemoteControlUI/Public/UI/RCFieldGroupType.h b/Engine/Plugins/VirtualProduction/RemoteControl/Source/RemoteControlUI/Public/UI/RCFieldGroupType.h
new file mode 100644
--- /dev/null
+++ b/Engine/Plugins/VirtualProduction/RemoteControl/Source/RemoteControlUI/Public/UI/RCFieldGroupType.h
@@ -0,0 +1,22 @@
+// Copyright Epic Games, Inc. All Rights Reserved.
+
+#pragma once
+
+#include "RCFieldGroupType.generated.h"
+
+/** Types of remote control field groups */
+UENUM()
+enum class ERCFieldGroupType : uint8
+{
+	/** No Grouping. Will be shown as a flat list */
+	None,
+
+	/** Fields with same property id are grouped together */
+	PropertyId,
+
+	/** Fields with same owner are grouped together */
+	Owner,
+};
+
+ENUM_CLASS_FLAGS(ERCFieldGroupType)
+
diff --git a/Engine/Plugins/VirtualProduction/RemoteControl/Source/RemoteControlUI/Public/UI/RCPanelExposedEntitiesListSettingsData.h b/Engine/Plugins/VirtualProduction/RemoteControl/Source/RemoteControlUI/Public/UI/RCPanelExposedEntitiesListSettingsData.h
new file mode 100644
--- /dev/null
+++ b/Engine/Plugins/VirtualProduction/RemoteControl/Source/RemoteControlUI/Public/UI/RCPanelExposedEntitiesListSettingsData.h
@@ -0,0 +1,23 @@
+// Copyright Epic Games, Inc. All Rights Reserved.
+
+#pragma once
+
+#include "RCFieldGroupType.h"
+#include "RCFieldGroupOrder.h"
+
+#include "RCPanelExposedEntitiesListSettingsData.generated.h"
+
+/** Data for exposed entities list settings */
+USTRUCT()
+struct FRCPanelExposedEntitiesListSettingsData
+{
+	GENERATED_BODY()
+
+	/** The field group type for the entity list. */
+	UPROPERTY()
+	ERCFieldGroupType FieldGroupType = ERCFieldGroupType::None;
+
+	/** Whether the field groups are expanded */
+	UPROPERTY()
+	ERCFieldGroupOrder FieldGroupOrder = ERCFieldGroupOrder::Ascending;
+};
diff --git a/Engine/Plugins/VirtualProduction/RemoteControlProtocolDMX/Source/RemoteControlProtocolDMX/Internal/Library/RemoteControlDMXControlledProperty.h b/Engine/Plugins/VirtualProduction/RemoteControlProtocolDMX/Source/RemoteControlProtocolDMX/Internal/Library/RemoteControlDMXControlledProperty.h
new file mode 100644
--- /dev/null
+++ b/Engine/Plugins/VirtualProduction/RemoteControlProtocolDMX/Source/RemoteControlProtocolDMX/Internal/Library/RemoteControlDMXControlledProperty.h
@@ -0,0 +1,50 @@
+// Copyright Epic Games, Inc. All Rights Reserved.
+
+#pragma once
+
+#include "Library/DMXEntityFixturePatch.h"
+#include "UObject/StructOnScope.h"
+
+enum class ERemoteControlDMXPatchGroupMode : uint8;
+struct FRemoteControlDMXProtocolEntity;
+struct FRemoteControlProperty;
+struct FRemoteControlProtocolEntity;
+
+namespace UE::RemoteControl::DMX
+{
+	/** Defines a Patch in a DMX Library of a Remote Control Preset along with the properties it controls */
+	class REMOTECONTROLPROTOCOLDMX_API FRemoteControlDMXControlledProperty
+		: public TSharedFromThis<FRemoteControlDMXControlledProperty>
+	{
+	public:
+		FRemoteControlDMXControlledProperty(const TSharedRef<const FRemoteControlProperty> InExposedProperty);
+
+		/** Returns the DMX protocol entities as structs on scope */
+		TArray<TSharedRef<TStructOnScope<FRemoteControlProtocolEntity>>> GetEntities() const { return Entities; }
+
+		/** Returns the owner object of this property */
+		const UObject* GetOwnerObject() const;
+
+		/** Returns the fixture patch of this property */
+		UDMXEntityFixturePatch* GetFixturePatch() const;
+
+		/** Returns the index of specified entity. Useful to enumerate attributes. The entity has to be owned by this DMX controlled property, checked. */
+		int32 GetEntityIndexChecked(const TSharedRef<TStructOnScope<FRemoteControlProtocolEntity>>& Entity);
+
+		/** The DMX controlled property */
+		const TSharedRef<const FRemoteControlProperty> ExposedProperty;
+
+	private:
+		/** Initializes new entities */
+		void InitializeNewEntities();
+
+		/** Makes sure all entities use unfied properties where required. */
+		void UnifyEntities();
+
+		/** Returns the subobject path for this property */
+		FString GetSubobjectPath() const;
+
+		/** The entities that control the property */
+		TArray<TSharedRef<TStructOnScope<FRemoteControlProtocolEntity>>> Entities;
+	};
+}
diff --git a/Engine/Plugins/VirtualProduction/RemoteControlProtocolDMX/Source/RemoteControlProtocolDMX/Internal/Library/RemoteControlDMXControlledPropertyPatch.h b/Engine/Plugins/VirtualProduction/RemoteControlProtocolDMX/Source/RemoteControlProtocolDMX/Internal/Library/RemoteControlDMXControlledPropertyPatch.h
new file mode 100644
--- /dev/null
+++ b/Engine/Plugins/VirtualProduction/RemoteControlProtocolDMX/Source/RemoteControlProtocolDMX/Internal/Library/RemoteControlDMXControlledPropertyPatch.h
@@ -0,0 +1,47 @@
+// Copyright Epic Games, Inc. All Rights Reserved.
+
+#pragma once
+
+#include "Templates/SharedPointer.h"
+
+class UDMXEntityFixturePatch;
+class UDMXEntityFixtureType;
+class UDMXLibrary;
+class URemoteControlDMXUserData;
+enum class ERemoteControlDMXPatchGroupMode : uint8;
+struct FDMXFixtureFunction;
+struct FRemoteControlProperty;
+struct FRemoteControlProtocolEntity;
+
+namespace UE::RemoteControl::DMX
+{
+	class FRemoteControlDMXControlledProperty;
+
+	/** A collection of DMX controlled properties that form a DMX fixture patch together */
+	class REMOTECONTROLPROTOCOLDMX_API FRemoteControlDMXControlledPropertyPatch
+		: public TSharedFromThis<FRemoteControlDMXControlledPropertyPatch>
+	{
+	public:
+		/** Constructs the DMX controlled properties patch. */
+		FRemoteControlDMXControlledPropertyPatch(
+			URemoteControlDMXUserData& DMXUserData, 
+			const TArray<TSharedRef<FRemoteControlDMXControlledProperty>>& InDMXControlledProperties);
+
+		~FRemoteControlDMXControlledPropertyPatch();
+
+		/** Returns the DMX controlled properties that form this patch */
+		const TArray<TSharedRef<FRemoteControlDMXControlledProperty>>& GetDMXControlledProperties() const { return DMXControlledProperties; }
+
+		/** Returns the fixture patch corresponding to this property patch */
+		UDMXEntityFixturePatch* GetFixturePatch() const;
+
+	private:
+#if WITH_EDITOR
+		/** Called when a fixture type changed. Useful to adopt changed properties from the fixture type. */
+		void OnFixtureTypeChanged(const UDMXEntityFixtureType* ChangedFixtureType);
+#endif // WITH_EDITOR
+
+		/** The DMX controlled properties that form this patch */
+		TArray<TSharedRef<FRemoteControlDMXControlledProperty>> DMXControlledProperties;
+	};
+}
diff --git a/Engine/Plugins/VirtualProduction/RemoteControlProtocolDMX/Source/RemoteControlProtocolDMX/Internal/Library/RemoteControlDMXLibraryProxy.h b/Engine/Plugins/VirtualProduction/RemoteControlProtocolDMX/Source/RemoteControlProtocolDMX/Internal/Library/RemoteControlDMXLibraryProxy.h
new file mode 100644
--- /dev/null
+++ b/Engine/Plugins/VirtualProduction/RemoteControlProtocolDMX/Source/RemoteControlProtocolDMX/Internal/Library/RemoteControlDMXLibraryProxy.h
@@ -0,0 +1,114 @@
+// Copyright Epic Games, Inc. All Rights Reserved.
+
+#pragma once
+
+#include "UObject/Object.h"
+
+#include "RemoteControlDMXLibraryProxy.generated.h"
+
+struct FRemoteControlProperty;
+struct FRemoteControlProtocolEntity;
+class UDMXLibrary;
+class URemoteControlDMXUserData;
+class URemoteControlPreset;
+namespace UE::RemoteControl::DMX 
+{ 
+	class FRemoteControlDMXControlledProperty; 
+	class FRemoteControlDMXControlledPropertyPatch;
+}
+
+DECLARE_MULTICAST_DELEGATE_OneParam(FRemoteControlDMXPrePropertyPatchesChanged, URemoteControlPreset* /** ChangedPreset */)
+
+/** Class responsible to maintain the DMX Library of a Remote Control Preset */
+UCLASS(Transient)
+class REMOTECONTROLPROTOCOLDMX_API URemoteControlDMXLibraryProxy : public UObject
+{
+	GENERATED_BODY()
+
+	using FRemoteControlDMXControlledProperty = UE::RemoteControl::DMX::FRemoteControlDMXControlledProperty;
+	using FRemoteControlDMXControlledPropertyPatch = UE::RemoteControl::DMX::FRemoteControlDMXControlledPropertyPatch;
+
+public:
+	//~ Begin UObject interface
+	virtual void PostInitProperties() override;
+	virtual void PostLoad() override;
+	//~ End UObject interface
+
+	/** Gets the DMX library of this proxy */
+	UDMXLibrary* GetDMXLibrary() const;
+
+	/** Returns the DMX controlled property patches in this proxy */
+	TArray<TSharedRef<FRemoteControlDMXControlledPropertyPatch>> GetPropertyPatches() const { return PropertyPatches; }
+
+	/** Refreshes the proxy on end frame */
+	void RequestRefresh();
+
+	/** Refreshes the proxy */
+	void Refresh();
+
+	/** Resets the proxy so it no longer can receive DMX */
+	void Reset();
+
+#if WITH_EDITOR
+	/** Finds fixture patches maintained by this proxy that exceed universe size */
+	TArray<UDMXEntityFixturePatch*> FindPatchesThatExceedUniverseSize() const;
+
+	/** Returns a delegate broadcast before property patches are being changed */
+	static FRemoteControlDMXPrePropertyPatchesChanged& GetOnPrePropertyPatchesChanged() { return OnPrePropertyPatchesChanged; }
+
+	/** Returns a delegate broadcast after property patches were changed */
+	static FSimpleMulticastDelegate& GetOnPostPropertyPatchesChanged() { return OnPostPropertyPatchesChanged; }
+#endif // WITH_EDITOR
+
+private:
+	/** Updates the property patches used in this proxy */
+	void UpdatePropertyPatches();
+
+	/** Binds to the OnFixturePatchReceived event for all patches in use */
+	void BindOnFixturePatchesReceived();
+
+	/** Unbinds from the OnFixturePatchReceived event for all patches in use */
+	void UnbindOnFixturePatchesReceived();
+
+	/** Called when a fixture patch was received */
+	UFUNCTION()
+	void OnFixturePatchReceived(UDMXEntityFixturePatch* FixturePatch, const FDMXNormalizedAttributeValueMap& ValuePerAttribute);
+
+	/** Returns the current fixture patches in use */
+	TArray<UDMXEntityFixturePatch*> GetFixturePatches() const;
+
+	/** Returns exposed properties of the preset that owns this proxy */
+	TArray<TSharedRef<FRemoteControlProperty>> GetExposedProperties() const;
+
+	/** Returns the user data that holds this object */
+	URemoteControlDMXUserData& GetMutableDMXUserDataChecked();
+	
+	/** Returns the user data that holds this object */
+	const URemoteControlDMXUserData& GetDMXUserDataChecked() const;
+
+	/** Called when the remote control preset was fully loaded */
+	void OnPostLoadRemoteControlPreset(URemoteControlPreset* Preset);
+
+	/** Called when an entity was exposed or unexposed */
+	void OnEntityExposedOrUnexposed(URemoteControlPreset* Preset, const FGuid& EntityId);
+
+	/** Called when an entity was rebound */
+	void OnEntityRebind(const FGuid& EntityId);
+
+	/** Called when entities changed */
+	void OnEntitiesUpdated(URemoteControlPreset* Preset, const TSet<FGuid>& ModifiedEntities);
+
+	/** Current DMX controlled property patches */
+	TArray<TSharedRef<FRemoteControlDMXControlledPropertyPatch>> PropertyPatches;
+
+	/** Handle set when the proxy is about to refresh */
+	FDelegateHandle RefreshDelegateHandle;
+
+#if WITH_EDITOR
+	/** Delegate broadcast before property patches are being changed */
+	static FRemoteControlDMXPrePropertyPatchesChanged OnPrePropertyPatchesChanged;
+
+	/** Delegate broadcast after property patches were changed */
+	static FSimpleMulticastDelegate OnPostPropertyPatchesChanged;
+#endif // WITH_EDITOR
+};
diff --git a/Engine/Plugins/VirtualProduction/RemoteControlProtocolDMX/Source/RemoteControlProtocolDMX/Internal/RemoteControlDMXUserData.h b/Engine/Plugins/VirtualProduction/RemoteControlProtocolDMX/Source/RemoteControlProtocolDMX/Internal/RemoteControlDMXUserData.h
new file mode 100644
--- /dev/null
+++ b/Engine/Plugins/VirtualProduction/RemoteControlProtocolDMX/Source/RemoteControlProtocolDMX/Internal/RemoteControlDMXUserData.h
@@ -0,0 +1,115 @@
+// Copyright Epic Games, Inc. All Rights Reserved.
+
+#pragma once
+
+#include "Library/DMXEntityReference.h"
+#include "UObject/Object.h"
+
+#include "RemoteControlDMXUserData.generated.h"
+
+struct FRemoteControlProtocolEntity;
+class UDMXLibrary;
+class URemoteControlDMXLibraryProxy;
+class URemoteControlPreset;
+
+/** Defines how DMX Protocol Entities should be patched */
+UENUM()
+enum class ERemoteControlDMXPatchGroupMode : uint8
+{
+	/** Creates a patch per property */
+	GroupByProperty,
+
+	/** Creates a patch per property owner object */
+	GroupByOwner,
+};
+
+DECLARE_EVENT_OneParam(URemoteControlDMXUserData, FOnRemoteControlDMXUserDataChangedProperty, FPropertyChangedEvent&);
+
+/** DMX user data for a Remote Control Preset */
+UCLASS()
+class REMOTECONTROLPROTOCOLDMX_API URemoteControlDMXUserData : public UObject
+{
+	GENERATED_BODY()
+
+	// Allow per preset editor settings to store itself in this DMX user data
+	friend class URemoteControlDMXPerPresetEditorSettings;
+
+public:
+	URemoteControlDMXUserData();
+
+	//~ Begin UObject interface
+#if WITH_EDITOR
+	virtual void PostEditChangeProperty(FPropertyChangedEvent& PropertyChangedEvent) override;
+#endif // WITH_EDITOR
+	virtual void PostLoad() override;
+	virtual void BeginDestroy() override;
+	//~ End UObject interface 
+
+	/** Gets or creates DMX User Data in specified preset */
+	static URemoteControlDMXUserData* GetOrCreateDMXUserData(URemoteControlPreset* InPreset);
+
+	/** Returns the DMX Library held with this user data */
+	UDMXLibrary* GetDMXLibrary() const { return DMXLibrary; }
+
+	/** Returns the DMX Library Proxy */
+	URemoteControlDMXLibraryProxy* GetDMXLibraryProxy() const { return DMXLibraryProxy; }
+
+	/** Returns the remote control preset this user data resides in */
+	URemoteControlPreset* GetPreset() const;
+
+	/** Sets the patch group mode */
+	void SetPatchGroupMode(ERemoteControlDMXPatchGroupMode NewPatchGroupMode);
+
+	/** Returns the current patch mode */
+	ERemoteControlDMXPatchGroupMode GetPatchGroupMode() const { return PatchGroupMode; }
+
+	/** Returns true if patches should be auto assigned to a universe and address, otherwise patches are is editable in editor. */
+	bool IsAutoPatch() const { return bAutoPatch; }
+
+	/** Sets if auto patching is enabled */
+	void SetAutoPatchEnabled(bool bEnabled);
+
+	/** Sets the universe from where patches should be generated. */
+	void SetAutoAssignFromUniverse(const int32 NewAutoAssignFromUniverse);
+
+	/** Returns the universe from where patches should be generated. */
+	int32 GetAutoAssignFromUniverse() const;
+
+#if WITH_EDITOR
+	// Property name getters
+	static const FName GetDMXLibraryMemberNameChecked() { return GET_MEMBER_NAME_CHECKED(URemoteControlDMXUserData, DMXLibrary); }
+	static const FName GetAutoPatchMemberNameChecked() { return GET_MEMBER_NAME_CHECKED(URemoteControlDMXUserData, DMXLibrary); }
+#endif // WITH_EDITOR
+
+private:
+#if WITH_EDITOR
+	/** Upgrades assets created before 5.5 that did not use a DMX library */
+	void TryUpgradeFromLegacy();
+#endif // WITH_EDITOR
+
+	/** The DMX library the remote control preset uses */
+	UPROPERTY(EditAnywhere, Category = "DMX")
+	TObjectPtr<UDMXLibrary> DMXLibrary;
+
+	/** Defines how DMX Protocol Entities should be patched  */
+	UPROPERTY()
+	ERemoteControlDMXPatchGroupMode PatchGroupMode;
+
+	/** Proxy to handle the DMX Library */
+	UPROPERTY(Transient)
+	TObjectPtr<URemoteControlDMXLibraryProxy> DMXLibraryProxy;
+
+	/** The universe from where patches should be generated. */
+	UPROPERTY()
+	int32 AutoAssignFromUniverse = 1;
+
+	/** When checked, patches are auto assigned to a universe and address, otherwise the patch is editable in editor. */
+	UPROPERTY()
+	bool bAutoPatch = true;
+
+#if WITH_EDITORONLY_DATA
+	/** Generic UObject to store editor data per preset, see URemoteControlDMXPerPresetEditorSettings */
+	UPROPERTY()
+	TObjectPtr<UObject> PerPresetEditorSettings;
+#endif // WITH_EDITORONLY_DATA
+};
diff --git a/Engine/Plugins/VirtualProduction/RemoteControlProtocolDMX/Source/RemoteControlProtocolDMX/Private/Library/RemoteControlDMXControlledProperty.cpp b/Engine/Plugins/VirtualProduction/RemoteControlProtocolDMX/Source/RemoteControlProtocolDMX/Private/Library/RemoteControlDMXControlledProperty.cpp
new file mode 100644
--- /dev/null
+++ b/Engine/Plugins/VirtualProduction/RemoteControlProtocolDMX/Source/RemoteControlProtocolDMX/Private/Library/RemoteControlDMXControlledProperty.cpp
@@ -0,0 +1,144 @@
+// Copyright Epic Games, Inc. All Rights Reserved.
+
+#include "Library/RemoteControlDMXControlledProperty.h"
+
+#include "GameFramework/Actor.h"
+#include "RemoteControlField.h"
+#include "RemoteControlProtocolDMX.h"
+
+namespace UE::RemoteControl::DMX
+{
+	FRemoteControlDMXControlledProperty::FRemoteControlDMXControlledProperty(const TSharedRef<const FRemoteControlProperty> InExposedProperty)
+		: ExposedProperty(InExposedProperty)
+	{
+		Algo::TransformIf(ExposedProperty->ProtocolBindings, Entities,
+			[](const FRemoteControlProtocolBinding& Binding)
+			{
+				return
+					Binding.GetProtocolName() == FRemoteControlProtocolDMX::ProtocolName &&
+					Binding.GetRemoteControlProtocolEntityPtr().IsValid() && 
+					Binding.GetRemoteControlProtocolEntityPtr()->IsValid();
+			},
+			[](const FRemoteControlProtocolBinding& Binding)
+			{
+				return Binding.GetRemoteControlProtocolEntityPtr().ToSharedRef();
+			});
+
+		InitializeNewEntities();
+		UnifyEntities();
+	}
+
+	const UObject* FRemoteControlDMXControlledProperty::GetOwnerObject() const
+	{
+		const UObject* BoundObject = ExposedProperty->GetBoundObject();
+		if (!BoundObject)
+		{
+			return nullptr;
+		}
+
+		if (const AActor* OwnerActor = BoundObject->GetTypedOuter<AActor>())
+		{
+			return OwnerActor;
+		}
+		else if (const AActor* Actor = Cast<AActor>(BoundObject))
+		{
+			return Actor;
+		}
+		else
+		{
+			return BoundObject;
+		}
+	}
+
+	UDMXEntityFixturePatch* FRemoteControlDMXControlledProperty::GetFixturePatch() const
+	{
+		const FRemoteControlDMXProtocolEntity* FirstDMXEntity = !Entities.IsEmpty() && Entities[0]->IsValid() ? Entities[0]->Cast<FRemoteControlDMXProtocolEntity>() : nullptr;
+		return FirstDMXEntity ? FirstDMXEntity->ExtraSetting.FixturePatchReference.GetFixturePatch() : nullptr;
+	}
+
+	int32 FRemoteControlDMXControlledProperty::GetEntityIndexChecked(const TSharedRef<TStructOnScope<FRemoteControlProtocolEntity>>& Entity)
+	{
+		const int32 EntityIndex = Entities.IndexOfByKey(Entity);
+		check(EntityIndex != INDEX_NONE);
+
+		return EntityIndex;
+	}
+
+	void FRemoteControlDMXControlledProperty::InitializeNewEntities()
+	{
+		for (int32 EntityIndex = 0; EntityIndex < Entities.Num(); EntityIndex++)
+		{
+			FRemoteControlDMXProtocolEntity* DMXEntity = Entities[EntityIndex]->IsValid() ? Entities[EntityIndex]->Cast<FRemoteControlDMXProtocolEntity>() : nullptr;
+			if (!DMXEntity)
+			{
+				continue;
+			}
+
+			if (DMXEntity->ExtraSetting.FunctionIndex == INDEX_NONE)
+			{
+				const FString AttributeName = [EntityIndex, this]()
+					{
+						if (Entities.Num() == 1)
+						{
+							return FString::Printf(TEXT("%s.%s"), *GetSubobjectPath(), *ExposedProperty->FieldPathInfo.ToString());
+						}
+						else
+						{
+							return FString::Printf(TEXT("%s.%s%i"), *GetSubobjectPath(), *ExposedProperty->FieldPathInfo.ToString(), EntityIndex + 1);
+						}
+					}();
+
+				DMXEntity->ExtraSetting.AttributeName = *AttributeName;
+				DMXEntity->ExtraSetting.FunctionIndex = EntityIndex;
+			}
+		}
+	}
+
+	void FRemoteControlDMXControlledProperty::UnifyEntities()
+	{
+		const FRemoteControlDMXProtocolEntity* FirstDMXEntity = !Entities.IsEmpty() && Entities[0]->IsValid() ? Entities[0]->Cast<FRemoteControlDMXProtocolEntity>() : nullptr;
+		if (!FirstDMXEntity)
+		{
+			return;
+		}
+
+		UDMXEntityFixturePatch* UnifiedFixturePatch = FirstDMXEntity->ExtraSetting.FixturePatchReference.GetFixturePatch();
+		const int32 UnifiedFunctionIndex = FirstDMXEntity->ExtraSetting.FunctionIndex;
+		const bool bUnifiedIsPrimaryPatch = FirstDMXEntity->ExtraSetting.bIsPrimaryPatch;
+
+		for (const TSharedRef<TStructOnScope<FRemoteControlProtocolEntity>>& Entity : Entities)
+		{
+			FRemoteControlDMXProtocolEntity* DMXEntity = Entity->IsValid() ? Entity->Cast<FRemoteControlDMXProtocolEntity>() : nullptr;
+			if (!DMXEntity || FirstDMXEntity == FirstDMXEntity)
+			{
+				continue;
+			}
+
+			// Only update if the properties differ, as this will trigger another rebuild of the library proxy that uses this object
+			if (DMXEntity->ExtraSetting.FixturePatchReference.GetFixturePatch() != UnifiedFixturePatch ||
+				DMXEntity->ExtraSetting.FunctionIndex != UnifiedFunctionIndex ||
+				DMXEntity->ExtraSetting.bIsPrimaryPatch != bUnifiedIsPrimaryPatch)
+			{
+				ensureMsgf(0, TEXT("Found entities of the same property with differing properties. Adopting property from first entity to mend."));
+
+				DMXEntity->ExtraSetting.FixturePatchReference = UnifiedFixturePatch;
+				DMXEntity->ExtraSetting.FunctionIndex = UnifiedFunctionIndex;
+				DMXEntity->ExtraSetting.bIsPrimaryPatch = bUnifiedIsPrimaryPatch;
+			}
+		}
+	}
+
+	FString FRemoteControlDMXControlledProperty::GetSubobjectPath() const
+	{
+		if (const UObject* OwnerObject = GetOwnerObject())
+		{
+			const FString OwnerName = OwnerObject->GetFName().ToString();
+			const FString BindingPath = ExposedProperty->GetLastBindingPath().ToString();
+
+			const int32 OwnerNameIndex = BindingPath.Find(OwnerName + TEXT("."));
+			return BindingPath.RightChop(OwnerNameIndex + OwnerName.Len() + 1);
+		}
+
+		return TEXT("");
+	}
+}
diff --git a/Engine/Plugins/VirtualProduction/RemoteControlProtocolDMX/Source/RemoteControlProtocolDMX/Private/Library/RemoteControlDMXControlledPropertyPatch.cpp b/Engine/Plugins/VirtualProduction/RemoteControlProtocolDMX/Source/RemoteControlProtocolDMX/Private/Library/RemoteControlDMXControlledPropertyPatch.cpp
new file mode 100644
--- /dev/null
+++ b/Engine/Plugins/VirtualProduction/RemoteControlProtocolDMX/Source/RemoteControlProtocolDMX/Private/Library/RemoteControlDMXControlledPropertyPatch.cpp
@@ -0,0 +1,83 @@
+// Copyright Epic Games, Inc. All Rights Reserved.
+
+#include "Library/RemoteControlDMXControlledPropertyPatch.h"
+
+#include "Algo/Find.h"
+#include "Library/DMXEntityFixturePatch.h"
+#include "Library/DMXEntityFixtureType.h"
+#include "Library/RemoteControlDMXControlledProperty.h"
+#include "Library/RemoteControlDMXControlledPropertyPatch.h"
+#include "RemoteControlDMXUserData.h"
+#include "RemoteControlProtocolDMX.h"
+
+#define LOCTEXT_NAMESPACE "RemoteControlDMXControlledPropertyPatch"
+
+namespace UE::RemoteControl::DMX
+{
+	FRemoteControlDMXControlledPropertyPatch::FRemoteControlDMXControlledPropertyPatch(
+		URemoteControlDMXUserData& DMXUserData, 
+		const TArray<TSharedRef<FRemoteControlDMXControlledProperty>>& InDMXControlledProperties)
+		: DMXControlledProperties(InDMXControlledProperties)
+	{
+#if WITH_EDITOR
+		// Listen to fixture type changes in editor
+		UDMXEntityFixtureType::GetOnFixtureTypeChanged().AddRaw(this, &FRemoteControlDMXControlledPropertyPatch::OnFixtureTypeChanged);
+#endif // WITH_EDITOR
+	}
+
+	FRemoteControlDMXControlledPropertyPatch::~FRemoteControlDMXControlledPropertyPatch()
+	{
+		UDMXEntityFixtureType::GetOnFixtureTypeChanged().RemoveAll(this);
+	}
+
+	UDMXEntityFixturePatch* FRemoteControlDMXControlledPropertyPatch::GetFixturePatch() const
+	{
+		return DMXControlledProperties.IsEmpty() ? nullptr : DMXControlledProperties[0]->GetFixturePatch();
+	}
+
+#if WITH_EDITOR
+	void FRemoteControlDMXControlledPropertyPatch::OnFixtureTypeChanged(const UDMXEntityFixtureType* ChangedFixtureType)
+	{
+		UDMXEntityFixturePatch* FixturePatch = GetFixturePatch();
+		if (!FixturePatch || FixturePatch->GetFixtureType() != ChangedFixtureType)
+		{
+			return;
+		}
+
+		// Adopt attribute names, data type and signal format if they changed in the library
+		for (const TSharedRef<FRemoteControlDMXControlledProperty>& DMXControlledProperty : DMXControlledProperties)
+		{
+			for (const TSharedRef<TStructOnScope<FRemoteControlProtocolEntity>>& Entity : DMXControlledProperty->GetEntities())
+			{
+				FRemoteControlDMXProtocolEntity* DMXEntity = Entity->IsValid() ? Entity->Cast<FRemoteControlDMXProtocolEntity>() : nullptr;
+				const FDMXFixtureMode* ModePtr = FixturePatch ? FixturePatch->GetActiveMode() : nullptr;
+				if (!DMXEntity || !ModePtr)
+				{
+					continue;
+				}
+
+				const FDMXFixtureFunction* FunctionPtr = Algo::FindBy(ModePtr->Functions, DMXEntity->ExtraSetting.AttributeName, &FDMXFixtureFunction::Attribute);
+				if (FunctionPtr)
+				{
+					// If the attribute exists, update the function index
+					DMXEntity->ExtraSetting.FunctionIndex = ModePtr->Functions.IndexOfByPredicate(
+						[&FunctionPtr](const FDMXFixtureFunction& Function)
+						{
+							return &Function == FunctionPtr;
+						});
+				}
+
+				if (ensureMsgf(ModePtr->Functions.IsValidIndex(DMXEntity->ExtraSetting.FunctionIndex),
+					TEXT("Unexpected cannot find function for remote control protocol DMX entity. Cannot update entity.")))
+				{
+					DMXEntity->ExtraSetting.AttributeName = ModePtr->Functions[DMXEntity->ExtraSetting.FunctionIndex].Attribute.Name;
+					DMXEntity->ExtraSetting.DataType = ModePtr->Functions[DMXEntity->ExtraSetting.FunctionIndex].DataType;
+					DMXEntity->ExtraSetting.bUseLSB = ModePtr->Functions[DMXEntity->ExtraSetting.FunctionIndex].bUseLSBMode;
+				}
+			}
+		}
+	}
+#endif // WITH_EDITOR
+}
+
+#undef LOCTEXT_NAMESPACE
diff --git a/Engine/Plugins/VirtualProduction/RemoteControlProtocolDMX/Source/RemoteControlProtocolDMX/Private/Library/RemoteControlDMXLibraryProxy.cpp b/Engine/Plugins/VirtualProduction/RemoteControlProtocolDMX/Source/RemoteControlProtocolDMX/Private/Library/RemoteControlDMXLibraryProxy.cpp
new file mode 100644
--- /dev/null
+++ b/Engine/Plugins/VirtualProduction/RemoteControlProtocolDMX/Source/RemoteControlProtocolDMX/Private/Library/RemoteControlDMXLibraryProxy.cpp
@@ -0,0 +1,351 @@
+// Copyright Epic Games, Inc. All Rights Reserved.
+
+#include "Library/RemoteControlDMXLibraryProxy.h"
+
+#include "Algo/AnyOf.h"
+#include "IRemoteControlProtocolModule.h"
+#include "Library/DMXEntityFixturePatch.h"
+#include "Library/DMXLibrary.h"
+#include "Library/RemoteControlDMXControlledProperty.h"
+#include "Library/RemoteControlDMXControlledPropertyPatch.h"
+#include "Misc/CoreDelegates.h"
+#include "RemoteControlDMXLog.h"
+#include "RemoteControlDMXUserData.h"
+#include "RemoteControlField.h"
+#include "RemoteControlPreset.h"
+#include "RemoteControlProtocolDMX.h"
+
+#if WITH_EDITOR
+FRemoteControlDMXPrePropertyPatchesChanged URemoteControlDMXLibraryProxy::OnPrePropertyPatchesChanged;
+FSimpleMulticastDelegate URemoteControlDMXLibraryProxy::OnPostPropertyPatchesChanged;
+#endif 
+
+void URemoteControlDMXLibraryProxy::PostInitProperties()
+{
+	Super::PostInitProperties();
+
+	if (IsTemplate())
+	{
+		return;
+	}
+
+	URemoteControlPreset* Preset = Cast<URemoteControlPreset>(GetMutableDMXUserDataChecked().GetOuter());
+	if (!ensureMsgf(Preset, TEXT("Invalid outer preset for DMX library proxy. Cannot initialize proxy.")))
+	{
+		return;
+	}
+
+	URemoteControlPreset::OnPostLoadRemoteControlPreset.AddUObject(this, &URemoteControlDMXLibraryProxy::OnPostLoadRemoteControlPreset);
+	Preset->OnEntityExposed().AddUObject(this, &URemoteControlDMXLibraryProxy::OnEntityExposedOrUnexposed);
+	Preset->OnEntityUnexposed().AddUObject(this, &URemoteControlDMXLibraryProxy::OnEntityExposedOrUnexposed);
+	Preset->OnEntityRebind().AddUObject(this, &URemoteControlDMXLibraryProxy::OnEntityRebind);
+	Preset->OnEntitiesUpdated().AddUObject(this, &URemoteControlDMXLibraryProxy::OnEntitiesUpdated);
+}
+
+void URemoteControlDMXLibraryProxy::PostLoad()
+{
+	Super::PostLoad();
+
+	if (!IsTemplate())
+	{
+		UpdatePropertyPatches();
+	}
+}
+
+UDMXLibrary* URemoteControlDMXLibraryProxy::GetDMXLibrary() const
+{
+	URemoteControlDMXUserData* DMXUserData = Cast<URemoteControlDMXUserData>(GetOuter());
+	UDMXLibrary* DMXLibrary = DMXUserData ? DMXUserData->GetDMXLibrary() : nullptr;
+	if (!ensureMsgf(DMXUserData, TEXT("Invalid outer for URemoteControlDMXUserData. The proxy expects URemoteControlDMXUserData resp. URemoteControlPreset as its outers.")) ||
+		!ensureMsgf(DMXLibrary, TEXT("Unexpected URemoteControlDMXUserData has no valid DMX Library.")))
+	{
+		return nullptr;
+	}
+
+	return DMXLibrary;
+}
+
+void URemoteControlDMXLibraryProxy::RequestRefresh()
+{
+	// Refresh on the next tick
+	if (!RefreshDelegateHandle.IsValid())
+	{
+		UnbindOnFixturePatchesReceived();
+
+		RefreshDelegateHandle = FCoreDelegates::OnEndFrame.AddUObject(this, &URemoteControlDMXLibraryProxy::Refresh);
+	}
+}
+
+void URemoteControlDMXLibraryProxy::Refresh()
+{
+	using namespace UE::RemoteControl::DMX;
+
+	RefreshDelegateHandle.Reset();
+	FCoreDelegates::OnEndFrame.RemoveAll(this);
+
+	// Don't refresh if this object is no longer fully valid
+	if (!IsValid(this))
+	{
+		return;
+	}
+
+	URemoteControlDMXUserData* DMXUserData = Cast<URemoteControlDMXUserData>(GetOuter());
+	URemoteControlPreset* Preset = DMXUserData ? Cast<URemoteControlPreset>(DMXUserData->GetOuter()) : nullptr;
+	if (Preset)
+	{
+		UnbindOnFixturePatchesReceived();
+
+#if WITH_EDITOR
+		OnPrePropertyPatchesChanged.Broadcast(Preset);
+#endif 
+
+		UpdatePropertyPatches();
+		
+#if WITH_EDITOR
+		OnPostPropertyPatchesChanged.Broadcast();
+#endif 
+
+		BindOnFixturePatchesReceived();
+	}
+}
+
+void URemoteControlDMXLibraryProxy::Reset()
+{
+	RefreshDelegateHandle.Reset();
+	FCoreDelegates::OnEndFrame.RemoveAll(this);
+
+	PropertyPatches.Reset();
+}
+
+#if WITH_EDITOR
+TArray<UDMXEntityFixturePatch*> URemoteControlDMXLibraryProxy::FindPatchesThatExceedUniverseSize() const
+{
+	const URemoteControlDMXUserData& DMXUserData = GetDMXUserDataChecked();
+	const ERemoteControlDMXPatchGroupMode PatchGroupMode = DMXUserData.GetPatchGroupMode();
+
+	TArray<UDMXEntityFixturePatch*> PatchesThatExceedUniverseSize;
+	for (const TSharedRef<FRemoteControlDMXControlledPropertyPatch>& PropertyPatch : PropertyPatches)
+	{
+		UDMXEntityFixturePatch* FixturePatch = PropertyPatch->GetFixturePatch();
+		if (FixturePatch && FixturePatch->GetChannelSpan() > DMX_UNIVERSE_SIZE)
+		{
+			PatchesThatExceedUniverseSize.Add(FixturePatch);
+		}
+	}
+
+	return PatchesThatExceedUniverseSize;
+}
+#endif // WITH_EDITOR
+
+void URemoteControlDMXLibraryProxy::UpdatePropertyPatches()
+{	
+	URemoteControlDMXUserData& DMXUserData = GetMutableDMXUserDataChecked();
+	UDMXLibrary* DMXLibrary = DMXUserData.GetDMXLibrary();
+	if (!ensureMsgf(DMXLibrary, TEXT("Cannot create remote control DMX patches, library is invalid.")))
+	{
+		return;
+	}
+
+	// Create DMX Controlled Properties
+	TArray<TSharedRef<FRemoteControlProperty>> ExposedProperties = GetExposedProperties();
+
+	TArray<TSharedRef<FRemoteControlDMXControlledProperty>> DMXControlledProperties;
+	Algo::TransformIf(ExposedProperties, DMXControlledProperties,
+		[](const TSharedRef<FRemoteControlProperty>& ExposedProperty)
+		{			
+			const bool bValidPreset = ExposedProperty->GetOwner() != nullptr;
+			const bool bValidBoundObject = ExposedProperty->GetBoundObject() != nullptr;
+			const bool bHasAnyDMXBindings = Algo::AnyOf(ExposedProperty->ProtocolBindings, [](const FRemoteControlProtocolBinding& Binding)
+				{
+					return Binding.GetProtocolName() == FRemoteControlProtocolDMX::ProtocolName;
+				});
+
+			return 
+				bValidPreset &&
+				bValidBoundObject &&
+				bHasAnyDMXBindings;
+		},
+		[](const TSharedRef<FRemoteControlProperty>& ExposedProperty)
+		{
+			return MakeShared<FRemoteControlDMXControlledProperty>(ExposedProperty);
+		});
+
+	// Group entities depending on the patch mode
+	const ERemoteControlDMXPatchGroupMode PatchGroupMode = DMXUserData.GetPatchGroupMode();
+
+	if (PatchGroupMode == ERemoteControlDMXPatchGroupMode::GroupByProperty)
+	{
+		TArray<TSharedRef<FRemoteControlDMXControlledPropertyPatch>> PerPropertyPatches;
+		Algo::Transform(DMXControlledProperties, PerPropertyPatches,
+			[&DMXUserData](const TSharedRef<FRemoteControlDMXControlledProperty>& DMXControlledProperty)
+			{
+				return MakeShared<FRemoteControlDMXControlledPropertyPatch>(DMXUserData, TArray<TSharedRef<FRemoteControlDMXControlledProperty>>({ DMXControlledProperty }));
+			});
+
+		PropertyPatches = PerPropertyPatches;
+	}
+	else if (PatchGroupMode == ERemoteControlDMXPatchGroupMode::GroupByOwner)
+	{
+		TMap<const UObject*, TArray<TSharedRef<FRemoteControlDMXControlledProperty>>> OwnerToPropertiesMap;
+		for (const TSharedRef<FRemoteControlDMXControlledProperty>& DMXControlledProperty : DMXControlledProperties)
+		{
+			const UObject* Owner = DMXControlledProperty->ExposedProperty->GetBoundObject();
+			if (ensureMsgf(Owner, TEXT("Cannot group property by owner. Property has no valid bound object.")))
+			{
+				OwnerToPropertiesMap.FindOrAdd(Owner).Add(DMXControlledProperty);
+			}
+		}
+
+		TArray<TSharedRef<FRemoteControlDMXControlledPropertyPatch>> PatchesGroupedByOwner;
+		Algo::Transform(OwnerToPropertiesMap, PatchesGroupedByOwner, [&DMXUserData](const TTuple<const UObject*, TArray<TSharedRef<FRemoteControlDMXControlledProperty>>>& ObjectToPropertiesPair)
+			{
+				return MakeShared<FRemoteControlDMXControlledPropertyPatch>(DMXUserData, ObjectToPropertiesPair.Value);
+			});
+
+		PropertyPatches = PatchesGroupedByOwner;
+	}
+	else
+	{
+		ensureMsgf(0, TEXT("Unhandled patch mode, cannot create DMX patch for Remote Control Preset"));
+	}
+}
+
+void URemoteControlDMXLibraryProxy::BindOnFixturePatchesReceived()
+{
+	TArray<UDMXEntityFixturePatch*> FixturePatches = GetFixturePatches();
+	for (UDMXEntityFixturePatch* FixturePatch : FixturePatches)
+	{
+		if (!FixturePatch ||
+			FixturePatch->OnFixturePatchReceivedDMX.Contains(this, GET_FUNCTION_NAME_CHECKED(URemoteControlDMXLibraryProxy, OnFixturePatchReceived)))
+		{
+			continue;
+		}
+
+		FixturePatch->OnFixturePatchReceivedDMX.AddDynamic(this, &URemoteControlDMXLibraryProxy::OnFixturePatchReceived);
+	}
+}
+
+void URemoteControlDMXLibraryProxy::UnbindOnFixturePatchesReceived()
+{
+	TArray<UDMXEntityFixturePatch*> FixturePatches = GetFixturePatches();
+	for (UDMXEntityFixturePatch* FixturePatch : FixturePatches)
+	{
+		if (!FixturePatch ||
+			!FixturePatch->OnFixturePatchReceivedDMX.Contains(this, GET_FUNCTION_NAME_CHECKED(URemoteControlDMXLibraryProxy, OnFixturePatchReceived)))
+		{
+			continue;
+		}
+
+		FixturePatch->OnFixturePatchReceivedDMX.RemoveAll(this);
+	}
+}
+
+void URemoteControlDMXLibraryProxy::OnFixturePatchReceived(UDMXEntityFixturePatch* FixturePatch, const FDMXNormalizedAttributeValueMap& ValuePerAttribute)
+{
+	const TSharedPtr<IRemoteControlProtocol> DMXProtocol = IRemoteControlProtocolModule::Get().GetProtocolByName(FRemoteControlProtocolDMX::ProtocolName);
+	if (!DMXProtocol.IsValid())
+	{
+		return;
+	}
+
+	for (const TSharedRef<FRemoteControlDMXControlledPropertyPatch>& PropertyPatch : PropertyPatches)
+	{
+		if (PropertyPatch->GetFixturePatch() != FixturePatch)
+		{
+			continue;
+		}
+
+		for (const TSharedRef<FRemoteControlDMXControlledProperty>& DMXControlledProperty : PropertyPatch->GetDMXControlledProperties())
+		{
+			for (const TSharedRef<TStructOnScope<FRemoteControlProtocolEntity>>& Entity : DMXControlledProperty->GetEntities())
+			{
+				FRemoteControlDMXProtocolEntity* DMXEntity = Entity->IsValid() ? Entity->Cast<FRemoteControlDMXProtocolEntity>() : nullptr;
+				if (!DMXEntity)
+				{
+					continue;
+				}
+
+				bool bSuccess = false;
+				const int32 DMXValue = FixturePatch->GetAttributeValue(DMXEntity->ExtraSetting.AttributeName, bSuccess);
+				if (bSuccess)
+				{
+					DMXProtocol->QueueValue(Entity, DMXValue);
+				}
+			}
+		}
+	}
+}
+
+TArray<UDMXEntityFixturePatch*> URemoteControlDMXLibraryProxy::GetFixturePatches() const
+{
+	TArray<UDMXEntityFixturePatch*> FixturePatches;
+	Algo::TransformIf(PropertyPatches, FixturePatches,
+		[](const TSharedRef<FRemoteControlDMXControlledPropertyPatch>& PropertyPatch)
+		{
+			return PropertyPatch->GetFixturePatch() != nullptr;
+		},
+		[](const TSharedRef<FRemoteControlDMXControlledPropertyPatch>& PropertyPatch)
+		{
+			return PropertyPatch->GetFixturePatch();
+		});
+
+	return FixturePatches;
+}
+
+TArray<TSharedRef<FRemoteControlProperty>> URemoteControlDMXLibraryProxy::GetExposedProperties() const
+{
+	URemoteControlDMXUserData* DMXUserData = Cast<URemoteControlDMXUserData>(GetOuter());
+	URemoteControlPreset* Preset = DMXUserData ? Cast<URemoteControlPreset>(DMXUserData->GetOuter()) : nullptr;
+
+	UDMXLibrary* DMXLibrary = GetDMXLibrary();
+	if (!ensureMsgf(DMXUserData, TEXT("Invalid outer for URemoteControlDMXUserData. The proxy expects URemoteControlDMXUserData resp. URemoteControlPreset as its outers.")) ||
+		!ensureMsgf(Preset, TEXT("Unexpected URemoteControlDMXUserData has no valid outer DMX User Data.")))
+	{
+		return {};
+	}
+
+	const TArray<TWeakPtr<FRemoteControlProperty>> WeakExposedProperties = Preset->GetExposedEntities<FRemoteControlProperty>();
+
+	TArray<TSharedRef<FRemoteControlProperty>> ExposedProperties;
+	Algo::TransformIf(WeakExposedProperties, ExposedProperties,
+		[](const TWeakPtr<FRemoteControlProperty>& WeakExposedProperty)
+		{
+			return WeakExposedProperty.IsValid();
+		},
+		[](const TWeakPtr<FRemoteControlProperty>& WeakExposedProperty)
+		{
+			return WeakExposedProperty.Pin().ToSharedRef();
+		});
+
+	return ExposedProperties;
+}
+
+const URemoteControlDMXUserData& URemoteControlDMXLibraryProxy::GetDMXUserDataChecked() const
+{
+	return *CastChecked<URemoteControlDMXUserData>(GetOuter());
+}
+
+URemoteControlDMXUserData& URemoteControlDMXLibraryProxy::GetMutableDMXUserDataChecked()
+{
+	return *CastChecked<URemoteControlDMXUserData>(GetOuter());
+}
+
+void URemoteControlDMXLibraryProxy::OnPostLoadRemoteControlPreset(URemoteControlPreset* Preset)
+{
+	RequestRefresh();
+}
+
+void URemoteControlDMXLibraryProxy::OnEntityExposedOrUnexposed(URemoteControlPreset* Preset, const FGuid& EntityId)
+{
+	RequestRefresh();
+}
+
+void URemoteControlDMXLibraryProxy::OnEntityRebind(const FGuid& EntityId)
+{
+	RequestRefresh();
+}
+
+void URemoteControlDMXLibraryProxy::OnEntitiesUpdated(URemoteControlPreset* Preset, const TSet<FGuid>& ModifiedEntities)
+{
+	RequestRefresh();
+}
diff --git a/Engine/Plugins/VirtualProduction/RemoteControlProtocolDMX/Source/RemoteControlProtocolDMX/Private/RemoteControlDMXLog.h b/Engine/Plugins/VirtualProduction/RemoteControlProtocolDMX/Source/RemoteControlProtocolDMX/Private/RemoteControlDMXLog.h
new file mode 100644
--- /dev/null
+++ b/Engine/Plugins/VirtualProduction/RemoteControlProtocolDMX/Source/RemoteControlProtocolDMX/Private/RemoteControlDMXLog.h
@@ -0,0 +1,7 @@
+// Copyright Epic Games, Inc. All Rights Reserved.
+
+#pragma once
+
+#include "Logging/LogMacros.h"
+
+DECLARE_LOG_CATEGORY_CLASS(LogRemoteControlDMX, Log, All);
diff --git a/Engine/Plugins/VirtualProduction/RemoteControlProtocolDMX/Source/RemoteControlProtocolDMX/Private/RemoteControlDMXUserData.cpp b/Engine/Plugins/VirtualProduction/RemoteControlProtocolDMX/Source/RemoteControlProtocolDMX/Private/RemoteControlDMXUserData.cpp
new file mode 100644
--- /dev/null
+++ b/Engine/Plugins/VirtualProduction/RemoteControlProtocolDMX/Source/RemoteControlProtocolDMX/Private/RemoteControlDMXUserData.cpp
@@ -0,0 +1,235 @@
+// Copyright Epic Games, Inc. All Rights Reserved.
+
+#include "RemoteControlDMXUserData.h"
+
+#include "Algo/AnyOf.h"
+#include "Algo/Find.h"
+#include "Library/DMXEntityFixturePatch.h"
+#include "Library/DMXLibrary.h"
+#include "Library/RemoteControlDMXLibraryProxy.h"
+#include "RemoteControlProtocolDMX.h"
+#include "RemoteControlPreset.h"
+
+URemoteControlDMXUserData::URemoteControlDMXUserData()
+	: PatchGroupMode(ERemoteControlDMXPatchGroupMode::GroupByOwner)
+{
+	DMXLibrary = CreateDefaultSubobject<UDMXLibrary>("Internal");
+	DMXLibraryProxy = CreateDefaultSubobject<URemoteControlDMXLibraryProxy>("DMXLibraryProxy");
+}
+
+#if WITH_EDITOR
+void URemoteControlDMXUserData::PostEditChangeProperty(FPropertyChangedEvent& PropertyChangedEvent)
+{
+	Super::PostEditChangeProperty(PropertyChangedEvent);
+
+	if (PropertyChangedEvent.GetPropertyName() == GET_MEMBER_NAME_CHECKED(URemoteControlDMXUserData, DMXLibrary))
+	{
+		if (!DMXLibrary)
+		{
+			DMXLibrary = NewObject<UDMXLibrary>(this, "Internal", RF_Public | RF_Transactional);
+		}
+
+		check(DMXLibraryProxy);
+		DMXLibraryProxy->RequestRefresh();
+	}
+}
+#endif // WITH_EDITOR
+
+void URemoteControlDMXUserData::PostLoad()
+{
+	Super::PostLoad();
+
+	if (!IsTemplate())
+	{
+		check(DMXLibraryProxy);
+		DMXLibraryProxy->RequestRefresh();
+	}
+}
+
+void URemoteControlDMXUserData::BeginDestroy()
+{
+	Super::BeginDestroy();
+
+	if (!IsTemplate())
+	{
+		check(DMXLibraryProxy);
+		DMXLibraryProxy->Reset();
+	}
+}
+
+URemoteControlDMXUserData* URemoteControlDMXUserData::GetOrCreateDMXUserData(URemoteControlPreset* InPreset)
+{
+	if (!InPreset)
+	{
+		return nullptr;
+	}
+
+	const TObjectPtr<UObject>* DMXUserDataObjectPtr = Algo::FindByPredicate(InPreset->UserData, [](const TObjectPtr<UObject>& Object)
+		{
+			return Object && Object->GetClass() == URemoteControlDMXUserData::StaticClass();
+		});
+
+	URemoteControlDMXUserData* DMXUserData = [DMXUserDataObjectPtr, &InPreset]()
+		{
+			if (DMXUserDataObjectPtr)
+			{
+				return CastChecked<URemoteControlDMXUserData>(*DMXUserDataObjectPtr);
+			}
+			else
+			{
+				URemoteControlDMXUserData* NewDMXUserData = NewObject<URemoteControlDMXUserData>(InPreset);
+				InPreset->UserData.Add(NewDMXUserData);
+
+#if WITH_EDITOR
+				NewDMXUserData->TryUpgradeFromLegacy();
+#endif // WITH_EDITOR
+
+				return NewDMXUserData;
+			}
+		}();
+
+	return DMXUserData;
+}
+
+URemoteControlPreset* URemoteControlDMXUserData::GetPreset() const
+{
+	return Cast<URemoteControlPreset>(GetOuter());
+}
+
+void URemoteControlDMXUserData::SetAutoPatchEnabled(bool bEnabled)
+{
+	bAutoPatch = bEnabled;
+
+	if (bAutoPatch)
+	{
+		check(DMXLibraryProxy);
+		DMXLibraryProxy->RequestRefresh();
+	}
+}
+
+int32 URemoteControlDMXUserData::GetAutoAssignFromUniverse() const
+{
+	// Only use the auto assign from universe property when in auto patch mode
+	return bAutoPatch ? AutoAssignFromUniverse : 1;
+}
+
+void URemoteControlDMXUserData::SetAutoAssignFromUniverse(const int32 NewAutoAssignFromUniverse)
+{
+	if (AutoAssignFromUniverse != NewAutoAssignFromUniverse)
+	{
+		AutoAssignFromUniverse = FMath::Max(1, NewAutoAssignFromUniverse);
+
+		check(DMXLibraryProxy);
+		DMXLibraryProxy->RequestRefresh();
+	}
+}
+
+void URemoteControlDMXUserData::SetPatchGroupMode(ERemoteControlDMXPatchGroupMode NewPatchGroupMode)
+{
+	PatchGroupMode = NewPatchGroupMode;
+	
+	check(DMXLibraryProxy);
+	DMXLibraryProxy->RequestRefresh();
+}
+
+#if WITH_EDITOR
+void URemoteControlDMXUserData::TryUpgradeFromLegacy()
+{
+	// No need to test if this is a template 
+	if (IsTemplate())
+	{
+		return;
+	}
+
+	URemoteControlPreset* Preset = Cast<URemoteControlPreset>(GetOuter());
+	if (!ensureMsgf(Preset, TEXT("Unexpected invalid outer for DMX user data. Cannot try to upgrade Remote Control preset for DMX.")))
+	{
+		return;
+	}
+
+	const TArray<TWeakPtr<FRemoteControlProperty>> WeakExposedProperties = Preset->GetExposedEntities<FRemoteControlProperty>();
+	const bool bIsLegacy = Algo::AnyOf(WeakExposedProperties, [](const TWeakPtr<FRemoteControlProperty>& WeakProperty)
+		{
+			const TSharedPtr< FRemoteControlProperty> Property = WeakProperty.Pin();
+			if (!Property.IsValid())
+			{
+				return false;
+			}
+
+			return Algo::AnyOf(Property->ProtocolBindings, [](const FRemoteControlProtocolBinding& Binding)
+				{
+					const TSharedPtr<TStructOnScope<FRemoteControlProtocolEntity>> Entity = Binding.GetRemoteControlProtocolEntityPtr();
+
+					const FRemoteControlDMXProtocolEntity* DMXEntity = Entity.IsValid() && Entity->IsValid() ?
+						Entity->Cast<FRemoteControlDMXProtocolEntity>() :
+						nullptr;
+
+					PRAGMA_DISABLE_DEPRECATION_WARNINGS
+					return 
+						DMXEntity && 
+						DMXEntity->ExtraSetting.Universe_DEPRECATED > -1 &&
+						DMXEntity->ExtraSetting.StartingChannel_DEPRECATED > -1;
+					PRAGMA_ENABLE_DEPRECATION_WARNINGS
+				});
+		});
+
+	if (bIsLegacy)
+	{
+		// Use per property manual patching mode for legacy remote control presets
+		bAutoPatch = false;
+		PatchGroupMode = ERemoteControlDMXPatchGroupMode::GroupByProperty;
+
+		check(DMXLibraryProxy);
+		DMXLibraryProxy->Refresh();
+		
+		for (const TWeakPtr<FRemoteControlProperty>& WeakProperty : WeakExposedProperties)
+		{
+			const TSharedPtr< FRemoteControlProperty> Property = WeakProperty.Pin();
+			if (!Property.IsValid())
+			{
+				continue;
+			}
+
+			for (const FRemoteControlProtocolBinding& Binding : Property->ProtocolBindings)
+			{
+				const TSharedPtr<TStructOnScope<FRemoteControlProtocolEntity>> Entity = Binding.GetRemoteControlProtocolEntityPtr();
+
+				FRemoteControlDMXProtocolEntity* DMXEntity = Entity.IsValid() && Entity->IsValid() ?
+					Entity->Cast<FRemoteControlDMXProtocolEntity>() :
+					nullptr;
+
+				UDMXEntityFixturePatch* FixturePatch = DMXEntity ? DMXEntity->ExtraSetting.FixturePatchReference.GetFixturePatch() : nullptr;
+				if (!DMXEntity || !FixturePatch)
+				{
+					continue;
+				}
+
+				PRAGMA_DISABLE_DEPRECATION_WARNINGS
+				const int32 Universe = DMXEntity->ExtraSetting.Universe_DEPRECATED;
+				const int32 StartingChannel = DMXEntity->ExtraSetting.StartingChannel_DEPRECATED;
+				PRAGMA_ENABLE_DEPRECATION_WARNINGS
+
+				const int32 EndingChannel = StartingChannel + FixturePatch->GetChannelSpan() - 1;
+
+				// Assign the fixture patch of the DMX entity as per legacy properties
+				if (Universe > 0 &&
+					Universe <= DMX_MAX_UNIVERSE &&
+					StartingChannel > 0 &&
+					StartingChannel <= DMX_UNIVERSE_SIZE &&
+					EndingChannel > 0 &&
+					EndingChannel <= DMX_UNIVERSE_SIZE)
+				{
+					FixturePatch->SetUniverseID(Universe);
+					FixturePatch->SetStartingChannel(StartingChannel);
+				}
+
+				// Reset legacy properties
+				PRAGMA_DISABLE_DEPRECATION_WARNINGS
+				DMXEntity->ExtraSetting.Universe_DEPRECATED = -1;
+				DMXEntity->ExtraSetting.StartingChannel_DEPRECATED = -1;
+				PRAGMA_ENABLE_DEPRECATION_WARNINGS
+			}
+		}
+	}
+}
+#endif
diff --git a/Engine/Plugins/VirtualProduction/RemoteControlProtocolDMX/Source/RemoteControlProtocolDMX/Private/RemoteControlProtocolDMX.cpp b/Engine/Plugins/VirtualProduction/RemoteControlProtocolDMX/Source/RemoteControlProtocolDMX/Private/RemoteControlProtocolDMX.cpp
--- a/Engine/Plugins/VirtualProduction/RemoteControlProtocolDMX/Source/RemoteControlProtocolDMX/Private/RemoteControlProtocolDMX.cpp
+++ b/Engine/Plugins/VirtualProduction/RemoteControlProtocolDMX/Source/RemoteControlProtocolDMX/Private/RemoteControlProtocolDMX.cpp
@@ -2,32 +2,21 @@
 
 #include "RemoteControlProtocolDMX.h"
 
-#include "DMXConversions.h"
-#include "Library/DMXEntityFixtureType.h"
-#include "RemoteControlLogger.h"
+#include "Algo/Find.h"
+#include "Library/RemoteControlDMXControlledProperty.h"
+#include "Library/RemoteControlDMXControlledPropertyPatch.h"
+#include "Library/RemoteControlDMXLibraryProxy.h"
+#include "RemoteControlDMXUserData.h"
+#include "RemoteControlPreset.h"
 #include "RemoteControlProtocolDMXObjectVersion.h"
-#include "RemoteControlProtocolDMXSettings.h"
-#include "IO/DMXInputPort.h"
-#include "IO/DMXPortManager.h"
-
-#if WITH_EDITOR
-#include "IRemoteControlProtocolWidgetsModule.h"
-#endif
 
 #define LOCTEXT_NAMESPACE "RemoteControlProtocolDMX"
 
-const FName FRemoteControlProtocolDMX::ProtocolName = TEXT("DMX");
-
-#if WITH_EDITOR
-
-namespace RemoteControlDMXProtocolColumns
+uint8 FRemoteControlDMXProtocolEntity::GetNumDMXChannels() const
 {
-	static FName Channel = TEXT("Channel");
-	static FName Universe = TEXT("Universe");
+	return static_cast<uint8>(ExtraSetting.DataType) + 1;
 }
 
-#endif // WITH_EDITOR
-
 uint8 FRemoteControlDMXProtocolEntity::GetRangePropertySize() const
 {
 	switch (ExtraSetting.DataType)
@@ -81,45 +70,77 @@ const FString& FRemoteControlDMXProtocolEntity::GetRangePropertyMaxValue() const
 	}
 }
 
-#if WITH_EDITOR
-
-void FRemoteControlDMXProtocolEntity::RegisterProperties()
+void FRemoteControlDMXProtocolEntity::Invalidate()
 {
-	EXPOSE_PROTOCOL_PROPERTY(RemoteControlDMXProtocolColumns::Universe, FRemoteControlDMXProtocolEntityExtraSetting, Universe);
-	
-	EXPOSE_PROTOCOL_PROPERTY(RemoteControlDMXProtocolColumns::Channel, FRemoteControlDMXProtocolEntityExtraSetting, StartingChannel);
+	if (URemoteControlDMXLibraryProxy* DMXLibraryProxy = GetDMXLibraryProxy(GetOwner().Get()))
+	{
+		DMXLibraryProxy->RequestRefresh();
+	}
 }
 
-#endif // WITH_EDITOR
-
-void FRemoteControlDMXProtocolEntity::Initialize()
+TArray<TSharedRef<TStructOnScope<FRemoteControlProtocolEntity>>> FRemoteControlDMXProtocolEntity::GetAllDMXProtocolEntitiesInPreset(URemoteControlPreset* Preset)
 {
-	// Handle port changes
-	FDMXPortManager::Get().OnPortsChanged.AddRaw(this, &FRemoteControlDMXProtocolEntity::UpdateInputPort);
+	using namespace UE::RemoteControl::DMX;
 
-	// Handle project setting changes
-	const URemoteControlProtocolDMXSettings* ProtocolDMXSettings = GetDefault<URemoteControlProtocolDMXSettings>();
-	ProtocolDMXSettings->GetOnRemoteControlProtocolDMXSettingsChanged().AddRaw(this, &FRemoteControlDMXProtocolEntity::UpdateInputPort);
+	URemoteControlDMXLibraryProxy* DMXLibraryProxy = GetDMXLibraryProxy(Preset);
+	if (!DMXLibraryProxy)
+	{
+		return {};
+	}
 
-    // Assign InputPortReference
-	UpdateInputPort();
+	TArray<TSharedRef<TStructOnScope<FRemoteControlProtocolEntity>>> AllEntities;
+	for (const TSharedRef<FRemoteControlDMXControlledPropertyPatch>& Patch : DMXLibraryProxy->GetPropertyPatches())
+	{
+		for (const TSharedRef<FRemoteControlDMXControlledProperty>& Property : Patch->GetDMXControlledProperties())
+		{
+			AllEntities.Append(Property->GetEntities());
+		}
+	}
+
+	return AllEntities;
 }
 
-void FRemoteControlDMXProtocolEntity::UpdateInputPort()
+TArray<TSharedRef<TStructOnScope<FRemoteControlProtocolEntity>>> FRemoteControlDMXProtocolEntity::FindEntitiesByProperty(const TSharedRef<FRemoteControlProperty>& Property)
 {
-	if (ExtraSetting.bUseDefaultInputPort || !ExtraSetting.InputPortId.IsValid() || !FDMXPortManager::Get().FindInputPortByGuid(ExtraSetting.InputPortId))
+	using namespace UE::RemoteControl::DMX;
+
+	URemoteControlDMXLibraryProxy* DMXLibraryProxy = GetDMXLibraryProxy(Property->GetOwner());
+	if (!DMXLibraryProxy)
 	{
-		URemoteControlProtocolDMXSettings* RemoteControlDMXSettings = GetMutableDefault<URemoteControlProtocolDMXSettings>();
-		if (RemoteControlDMXSettings)
-		{
-			const FGuid& DefaultInputPortId = RemoteControlDMXSettings->GetOrCreateDefaultInputPortId();
-			FDMXInputPortSharedPtr DefaultInputPort = FDMXPortManager::Get().FindInputPortByGuid(DefaultInputPortId);
-			if (DefaultInputPort.IsValid())
+		return {};
+	}
+	
+	for (const TSharedRef<FRemoteControlDMXControlledPropertyPatch>& Patch : DMXLibraryProxy->GetPropertyPatches())
+	{
+		const TSharedRef<FRemoteControlDMXControlledProperty>* DMXControlledPropertyPtr = Algo::FindByPredicate(Patch->GetDMXControlledProperties(),
+			[&Property](const TSharedRef<FRemoteControlDMXControlledProperty>& DMXControlledProperty)
 			{
-				ExtraSetting.InputPortId = DefaultInputPortId;
-			}
+				return DMXControlledProperty->ExposedProperty == Property;
+			});
+		
+		if (DMXControlledPropertyPtr)
+		{
+			return (*DMXControlledPropertyPtr)->GetEntities();
 		}
 	}
+
+	return {};
+}
+
+void FRemoteControlDMXProtocolEntity::BindDMX()
+{
+	if (URemoteControlDMXLibraryProxy* DMXLibraryProxy = GetDMXLibraryProxy(GetOwner().Get()))
+	{
+		DMXLibraryProxy->RequestRefresh();
+	}
+}
+
+void FRemoteControlDMXProtocolEntity::UnbindDMX()
+{
+	if (URemoteControlDMXLibraryProxy* DMXLibraryProxy = GetDMXLibraryProxy(GetOwner().Get()))
+	{
+		DMXLibraryProxy->RequestRefresh();
+	}
 }
 
 bool FRemoteControlDMXProtocolEntity::Serialize(FArchive& Ar)
@@ -132,55 +153,67 @@ bool FRemoteControlDMXProtocolEntity::Serialize(FArchive& Ar)
 
 void FRemoteControlDMXProtocolEntity::PostSerialize(const FArchive& Ar)
 {
+#if WITH_EDITOR
 	if (Ar.IsLoading())
 	{
 		if (Ar.CustomVer(FRemoteControlProtocolDMXObjectVersion::GUID) < FRemoteControlProtocolDMXObjectVersion::MoveRemoteControlProtocolDMXEntityPropertiesToExtraSettingStruct)
 		{
-			// Move relevant properties that were moved to the the ExtraSetting member in 5.0 to the ExtraSetting member, so they can be customized. 
 			PRAGMA_DISABLE_DEPRECATION_WARNINGS
-			ExtraSetting.bUseDefaultInputPort = bUseDefaultInputPort_DEPRECATED;
+			// Move relevant properties that were moved to the the ExtraSetting member in 5.0 to the ExtraSetting member, so they can be customized.
+			ExtraSetting.bUseDefaultInputPort_DEPRECATED = bUseDefaultInputPort_DEPRECATED;
 			ExtraSetting.bUseLSB = bUseLSB_DEPRECATED;
 			ExtraSetting.DataType = DataType_DEPRECATED;
-			ExtraSetting.InputPortId = InputPortId_DEPRECATED;
-			ExtraSetting.Universe = Universe_DEPRECATED;
+			ExtraSetting.InputPortId_DEPRECATED = InputPortId_DEPRECATED;
+			ExtraSetting.Universe_DEPRECATED = Universe_DEPRECATED;
 			PRAGMA_ENABLE_DEPRECATION_WARNINGS
 		}
 	}
+#endif
+}
+
+URemoteControlDMXLibraryProxy* FRemoteControlDMXProtocolEntity::GetDMXLibraryProxy(URemoteControlPreset* Preset)
+{
+	URemoteControlDMXUserData* DMXUserData = URemoteControlDMXUserData::GetOrCreateDMXUserData(Preset);
+	return DMXUserData ? DMXUserData->GetDMXLibraryProxy() : nullptr;
 }
 
+const FName FRemoteControlProtocolDMX::ProtocolName = "DMX";
+
+#if WITH_EDITOR
+const FName FRemoteControlProtocolDMX::PatchColumnName = "Patch";
+const FName FRemoteControlProtocolDMX::UniverseColumnName = "Universe";
+const FName FRemoteControlProtocolDMX::ChannelColumnName = "Channel";
+#endif // WITH_EDITOR
+
 void FRemoteControlProtocolDMX::Bind(FRemoteControlProtocolEntityPtr InRemoteControlProtocolEntityPtr)
 {
 	if (!ensure(InRemoteControlProtocolEntityPtr.IsValid()))
 	{
 		return;
 	}
-
+	
 	FRemoteControlDMXProtocolEntity* DMXProtocolEntity = InRemoteControlProtocolEntityPtr->CastChecked<FRemoteControlDMXProtocolEntity>();
+	DMXProtocolEntity->BindDMX();
 
-	FRemoteControlProtocolEntityWeakPtr* ExistingProtocolBindings = ProtocolsBindings.FindByPredicate([DMXProtocolEntity](const FRemoteControlProtocolEntityWeakPtr& InProtocolEntity)
+	FRemoteControlProtocolEntityWeakPtr* ExistingProtocolBindings = WeakProtocolsBindings.FindByPredicate([DMXProtocolEntity](const FRemoteControlProtocolEntityWeakPtr& InProtocolEntity)
 		{
 			if (const FRemoteControlProtocolEntityPtr& ProtocolEntity = InProtocolEntity.Pin())
 			{
 				const FRemoteControlDMXProtocolEntity* ComparedDMXProtocolEntity = ProtocolEntity->CastChecked<FRemoteControlDMXProtocolEntity>();
 
-				if (ComparedDMXProtocolEntity->ExtraSetting.Universe == DMXProtocolEntity->ExtraSetting.Universe &&
-					ComparedDMXProtocolEntity->ExtraSetting.StartingChannel == DMXProtocolEntity->ExtraSetting.StartingChannel &&
-					ComparedDMXProtocolEntity->GetPropertyId() == DMXProtocolEntity->GetPropertyId())
+				if (ComparedDMXProtocolEntity->GetPropertyId() == DMXProtocolEntity->GetPropertyId())
 				{
 					return true;
 				}
 			}
-			
+
 			return false;
 		});
 
-
-	if (ExistingProtocolBindings == nullptr)
+	if (ExistingProtocolBindings)
 	{
-		ProtocolsBindings.Emplace(InRemoteControlProtocolEntityPtr);
+		WeakProtocolsBindings.Emplace(InRemoteControlProtocolEntityPtr);
 	}
-
-	DMXProtocolEntity->Initialize();
 }
 
 void FRemoteControlProtocolDMX::Unbind(FRemoteControlProtocolEntityPtr InRemoteControlProtocolEntityPtr)
@@ -190,175 +223,41 @@ void FRemoteControlProtocolDMX::Unbind(FRemoteControlProtocolEntityPtr InRemoteC
 		return;
 	}
 
-	const FRemoteControlDMXProtocolEntity* DMXProtocolEntity = InRemoteControlProtocolEntityPtr->CastChecked<FRemoteControlDMXProtocolEntity>();
-	ProtocolsBindings.RemoveAllSwap(CreateProtocolComparator(DMXProtocolEntity->GetPropertyId()));
-}
-
-void FRemoteControlProtocolDMX::OnEndFrame()
-{
-	if (!ProtocolsBindings.Num())
-	{
-		return;
-	}
-
-	const TArray<FDMXInputPortSharedRef>& InputPorts = FDMXPortManager::Get().GetInputPorts();
-	for (const FRemoteControlProtocolEntityWeakPtr& ProtocolEntityWeakPtr : ProtocolsBindings)
-	{
-		if (const FRemoteControlProtocolEntityPtr& ProtocolEntity = ProtocolEntityWeakPtr.Pin())
-		{
-			FRemoteControlDMXProtocolEntity* DMXProtocolEntity = ProtocolEntity->CastChecked<FRemoteControlDMXProtocolEntity>();
-
-			const FGuid& PortId =  DMXProtocolEntity->ExtraSetting.InputPortId;				
-
-			const FDMXInputPortSharedRef* InputPortPtr = InputPorts.FindByPredicate([&PortId](const FDMXInputPortSharedRef& InputPort) {
-                return InputPort->GetPortGuid() == PortId;
-                });
-
-			if (InputPortPtr == nullptr)
-			{
-				break;
-			}
-
-			const FDMXInputPortSharedRef& InputPort = *InputPortPtr;
-
-			// Get universe DMX signal
-			InputPort->GameThreadGetDMXSignal(DMXProtocolEntity->ExtraSetting.Universe, DMXProtocolEntity->LastSignalPtr);
-			if (DMXProtocolEntity->LastSignalPtr.IsValid())
-			{
-				const FDMXSignalSharedPtr& LastSignalPtr = DMXProtocolEntity->LastSignalPtr;
-				const int32 DMXOffset = DMXProtocolEntity->ExtraSetting.StartingChannel - 1;
-				check(DMXOffset >= 0 && DMXOffset < DMX_UNIVERSE_SIZE);
-
-#if WITH_EDITOR
-				ProcessAutoBinding(ProtocolEntity);
-#endif
-				
-				ProcessAndApplyProtocolValue(LastSignalPtr, DMXOffset, ProtocolEntity);
-			}
-		}
-	}
+	FRemoteControlDMXProtocolEntity* DMXProtocolEntity = InRemoteControlProtocolEntityPtr->CastChecked<FRemoteControlDMXProtocolEntity>();
+	DMXProtocolEntity->UnbindDMX();
 
-	FRemoteControlProtocol::OnEndFrame();
+	WeakProtocolsBindings.RemoveAllSwap(CreateProtocolComparator(DMXProtocolEntity->GetPropertyId()));
 }
 
-void FRemoteControlProtocolDMX::ProcessAndApplyProtocolValue(const FDMXSignalSharedPtr& InSignal, int32 InDMXOffset, const FRemoteControlProtocolEntityPtr& InProtocolEntityPtr)
-{
-	if (!InSignal->ChannelData.IsValidIndex(InDMXOffset))
-	{
-		return;
-	}
-	
-	FRemoteControlDMXProtocolEntity* DMXProtocolEntity = InProtocolEntityPtr->CastChecked<FRemoteControlDMXProtocolEntity>();
-	const uint8* ChannelData = &InSignal->ChannelData[InDMXOffset];	
-	const uint8 NumChannelsToOccupy = FDMXConversions::GetSizeOfSignalFormat(DMXProtocolEntity->ExtraSetting.DataType);
-
-	if(DMXProtocolEntity->CacheDMXBuffer.Num() != NumChannelsToOccupy ||
-		FMemory::Memcmp(DMXProtocolEntity->CacheDMXBuffer.GetData(), ChannelData, NumChannelsToOccupy) != 0)
-	{
-		const uint32 DMXValue = UDMXEntityFixtureType::BytesToInt(DMXProtocolEntity->ExtraSetting.DataType, DMXProtocolEntity->ExtraSetting.bUseLSB, ChannelData);
-		
 #if WITH_EDITOR
-		FRemoteControlLogger::Get().Log(ProtocolName, [&InSignal, DMXValue]
-		{
-			return FText::Format(LOCTEXT("DMXEventLog","ExternUniverseID {0}, DMXValue {1}"), InSignal->ExternUniverseID, DMXValue);
-		});
-#endif
-
-		QueueValue(InProtocolEntityPtr, DMXValue);
+void FRemoteControlProtocolDMX::RegisterColumns()
+{
+	FRemoteControlProtocol::RegisterColumns();
 
-		// update cached buffer
-		DMXProtocolEntity->CacheDMXBuffer = TArray<uint8>(ChannelData, NumChannelsToOccupy);
-	}
+	REGISTER_COLUMN(FRemoteControlProtocolDMX::PatchColumnName,
+		LOCTEXT("RCPresetPatchColumnHeader", "Patch"),
+		ProtocolColumnConstants::ColumnSizeNormal);
 }
+#endif // WITH_EDITOR
 
-#if WITH_EDITOR
-void FRemoteControlProtocolDMX::ProcessAutoBinding(const FRemoteControlProtocolEntityPtr& InProtocolEntityPtr)
+void FRemoteControlProtocolDMX::UnbindAll()
 {
-	// Bind only in Editor
-	if (!GIsEditor)
-	{
-		return;
-	}
-
-	// Check if entity is valid
-	if (!InProtocolEntityPtr.IsValid())
+	for (const TWeakPtr<TStructOnScope<FRemoteControlProtocolEntity>>& WeakEntity : WeakProtocolsBindings)
 	{
-		return;
-	}
-
-	FRemoteControlDMXProtocolEntity* DMXProtocolEntity = InProtocolEntityPtr->CastChecked<FRemoteControlDMXProtocolEntity>();
-	
-	// Assign binding
-	IRemoteControlProtocolWidgetsModule& RCWidgetsModule = IRemoteControlProtocolWidgetsModule::Get();
-	const TSharedPtr<IRCProtocolBindingList> RCProtocolBindingList = RCWidgetsModule.GetProtocolBindingList();
-	if (RCProtocolBindingList.IsValid())
-	{
-		if (DMXProtocolEntity->GetBindingStatus() == ERCBindingStatus::Awaiting)
+		if (!WeakEntity.IsValid())
 		{
-			const TArray<uint8>& ChannelData = DMXProtocolEntity->LastSignalPtr->ChannelData;
-			int32 FoundChannelDifference = -1;
-			uint8 FoundChannelDifferenceValue = 0;
-
-			if (CacheUniverseDMXBuffer.Num())
-			{
-				if (ensure(ChannelData.Num() == CacheUniverseDMXBuffer.Num()))
-				{
-					// Compare buffers
-					for (int32 ChannelValueIndex = 0; ChannelValueIndex < ChannelData.Num(); ++ChannelValueIndex)
-					{
-						const uint8 SignalChannelValue = ChannelData[ChannelValueIndex];
-						if (CacheUniverseDMXBuffer[ChannelValueIndex] != SignalChannelValue)
-						{
-							FoundChannelDifference = ChannelValueIndex;
-							FoundChannelDifferenceValue = SignalChannelValue;
-							break;
-						}
-					}	
-				}
-			}
-
-			if  (FoundChannelDifference >= 0)
-			{
-				Unbind(InProtocolEntityPtr);
-				const int32 FinalChannelValue = FoundChannelDifference + 1;
-				DMXProtocolEntity->ExtraSetting.StartingChannel = FinalChannelValue;
-				Bind(InProtocolEntityPtr);
-
-				// Print to log
-				const FDMXSignalSharedPtr& LastSignalPtr = DMXProtocolEntity->LastSignalPtr;
-				FRemoteControlLogger::Get().Log(ProtocolName, [&LastSignalPtr, FinalChannelValue, FoundChannelDifferenceValue]
-				{
-					return FText::Format(
-						LOCTEXT("DMXEventLog2",
-								"AutoBinding new value. ExternUniverseID {0}, Channel {1}, New Value {2}"),
-						LastSignalPtr->ExternUniverseID, FinalChannelValue, FoundChannelDifferenceValue);
-				});
-			}
+			continue;
+		}
+		const TSharedRef<TStructOnScope<FRemoteControlProtocolEntity>> Entity = WeakEntity.Pin().ToSharedRef();
 
-			// Copy buffer
-			CacheUniverseDMXBuffer = DMXProtocolEntity->LastSignalPtr->ChannelData;
+		FRemoteControlDMXProtocolEntity* DMXEntity = Entity->IsValid() ? Entity->Cast<FRemoteControlDMXProtocolEntity>() : nullptr;
+		if (!DMXEntity)
+		{
+			continue;
 		}
 	}
-}
-
-void FRemoteControlProtocolDMX::RegisterColumns()
-{
-	FRemoteControlProtocol::RegisterColumns();
-
-	REGISTER_COLUMN(RemoteControlDMXProtocolColumns::Channel
-		, LOCTEXT("RCPresetChannelColumnHeader", "Channel")
-		, ProtocolColumnConstants::ColumnSizeMicro);
-
-	REGISTER_COLUMN(RemoteControlDMXProtocolColumns::Universe
-		, LOCTEXT("RCPresetUniverseColumnHeader", "Universe")
-		, ProtocolColumnConstants::ColumnSizeMicro);
-}
 
-#endif
-
-void FRemoteControlProtocolDMX::UnbindAll()
-{
-	ProtocolsBindings.Empty();
+	WeakProtocolsBindings.Empty();
 }
 
 #undef LOCTEXT_NAMESPACE
diff --git a/Engine/Plugins/VirtualProduction/RemoteControlProtocolDMX/Source/RemoteControlProtocolDMX/Private/RemoteControlProtocolDMXModule.cpp b/Engine/Plugins/VirtualProduction/RemoteControlProtocolDMX/Source/RemoteControlProtocolDMX/Private/RemoteControlProtocolDMXModule.cpp
--- a/Engine/Plugins/VirtualProduction/RemoteControlProtocolDMX/Source/RemoteControlProtocolDMX/Private/RemoteControlProtocolDMXModule.cpp
+++ b/Engine/Plugins/VirtualProduction/RemoteControlProtocolDMX/Source/RemoteControlProtocolDMX/Private/RemoteControlProtocolDMXModule.cpp
@@ -18,6 +18,7 @@ class FRemoteControlProtocolDMXModule : public IModuleInterface
 {
 public:
 	//~ Begin IModuleInterface
+
 	virtual void StartupModule() override
 	{
 		const IRemoteControlProtocolModule& RemoteControlProtocolModule = IRemoteControlProtocolModule::Get();
@@ -25,32 +26,10 @@ public:
 		{
 			IRemoteControlProtocolModule::Get().AddProtocol(FRemoteControlProtocolDMX::ProtocolName, MakeShared<FRemoteControlProtocolDMX>());
 		}
-		
-#if WITH_EDITOR
-		ISettingsModule* SettingsModule = FModuleManager::GetModulePtr<ISettingsModule>("Settings");
-
-		// Register DMX Remote Control global settings
-		if (SettingsModule)
-		{
-			SettingsModule->RegisterSettings("Project", "Plugins", "Remote Control DMX Protocol",
-				LOCTEXT("ProjectSettings_Label", "Remote Control DMX Protocol"),
-				LOCTEXT("ProjectSettings_Description", "Configure MIDI remote control plugin global settings"),
-				GetMutableDefault<URemoteControlProtocolDMXSettings>()
-			);
-		}
-#endif // WITH_EDITOR
 	}
 
 	virtual void ShutdownModule() override
 	{
-#if WITH_EDITOR
-		// Unregister MIDI Remote Control global settings
-		ISettingsModule* SettingsModule = FModuleManager::GetModulePtr<ISettingsModule>("Settings");
-		if (SettingsModule)
-		{
-			SettingsModule->UnregisterSettings("Project", "Plugins", "Remote Control DMX Protocol");
-		}
-#endif // WITH_EDITOR
 	}
 
 	//~ End IModuleInterface
diff --git a/Engine/Plugins/VirtualProduction/RemoteControlProtocolDMX/Source/RemoteControlProtocolDMX/Private/RemoteControlProtocolDMXSettings.cpp b/Engine/Plugins/VirtualProduction/RemoteControlProtocolDMX/Source/RemoteControlProtocolDMX/Private/RemoteControlProtocolDMXSettings.cpp
--- a/Engine/Plugins/VirtualProduction/RemoteControlProtocolDMX/Source/RemoteControlProtocolDMX/Private/RemoteControlProtocolDMXSettings.cpp
+++ b/Engine/Plugins/VirtualProduction/RemoteControlProtocolDMX/Source/RemoteControlProtocolDMX/Private/RemoteControlProtocolDMXSettings.cpp
@@ -2,47 +2,18 @@
 
 #include "RemoteControlProtocolDMXSettings.h"
 
-#include "IO/DMXInputPort.h"
-#include "IO/DMXPortManager.h"
-
-FSimpleMulticastDelegate URemoteControlProtocolDMXSettings::OnRemoteControlProtocolDMXSettingsChangedDelegate;
-
-#if WITH_EDITOR
-void URemoteControlProtocolDMXSettings::PostEditChangeProperty(struct FPropertyChangedEvent& PropertyChangedEvent)
-{
-	Super::PostEditChangeProperty(PropertyChangedEvent);
-
-	if (PropertyChangedEvent.ChangeType != EPropertyChangeType::Interactive)
-	{
-		OnRemoteControlProtocolDMXSettingsChangedDelegate.Broadcast();
-	}
-}
-#endif // WITH_EDITOR
+PRAGMA_DISABLE_DEPRECATION_WARNINGS
+FSimpleMulticastDelegate URemoteControlProtocolDMXSettings::OnRemoteControlProtocolDMXSettingsChangedDelegate_DEPRECATED;
 
 FGuid URemoteControlProtocolDMXSettings::GetOrCreateDefaultInputPortId()
 {
-	if (!DefaultInputPortId.IsValid() || !FDMXPortManager::Get().FindInputPortByGuid(DefaultInputPortId))
-	{
-		const TArray<FDMXInputPortSharedRef> InputPorts = FDMXPortManager::Get().GetInputPorts();
-		if (InputPorts.Num() > 0)
-		{
-#if WITH_EDITOR
-			Modify();
-			PreEditChange(URemoteControlProtocolDMXSettings::StaticClass()->FindPropertyByName(GET_MEMBER_NAME_CHECKED(URemoteControlProtocolDMXSettings, DefaultInputPortId)));
-#endif 
-
-			DefaultInputPortId = InputPorts[0]->GetPortGuid();
-
-#if WITH_EDITOR
-			PostEditChange();
-#endif
-		}
-	}
+	// DEPRECATED 5.5
 
-	return DefaultInputPortId;
+	return FGuid();
 }
 
 FSimpleMulticastDelegate& URemoteControlProtocolDMXSettings::GetOnRemoteControlProtocolDMXSettingsChanged()
 {
-	return OnRemoteControlProtocolDMXSettingsChangedDelegate;
+	return OnRemoteControlProtocolDMXSettingsChangedDelegate_DEPRECATED;
 }
+PRAGMA_ENABLE_DEPRECATION_WARNINGS
diff --git a/Engine/Plugins/VirtualProduction/RemoteControlProtocolDMX/Source/RemoteControlProtocolDMX/Public/RemoteControlProtocolDMX.h b/Engine/Plugins/VirtualProduction/RemoteControlProtocolDMX/Source/RemoteControlProtocolDMX/Public/RemoteControlProtocolDMX.h
--- a/Engine/Plugins/VirtualProduction/RemoteControlProtocolDMX/Source/RemoteControlProtocolDMX/Public/RemoteControlProtocolDMX.h
+++ b/Engine/Plugins/VirtualProduction/RemoteControlProtocolDMX/Source/RemoteControlProtocolDMX/Public/RemoteControlProtocolDMX.h
@@ -5,36 +5,53 @@
 #include "RemoteControlProtocol.h"
 
 #include "DMXProtocolCommon.h"
-#include "RemoteControlProtocolBinding.h"
 #include "IO/DMXInputPortReference.h"
 #include "Library/DMXEntityFixtureType.h"
+#include "RemoteControlProtocolBinding.h"
 
 #include "RemoteControlProtocolDMX.generated.h"
 
 class FRemoteControlProtocolDMX;
+class URemoteControlDMXLibraryProxy;
+struct FRemoteControlProperty;
 
 /**
- * Using as an inner struct for details customization.
- * Useful to have type customization for the struct
+ * An inner struct holding DMX specific data.
+ * Useful to have type customization for the struct.
  */
 USTRUCT()
 struct FRemoteControlDMXProtocolEntityExtraSetting
 {
 	GENERATED_BODY();
 
+	/** Reference to the fixture patch this binding uses */
+	UPROPERTY()
+	FDMXEntityFixturePatchRef FixturePatchReference;
+
+#if WITH_EDITORONLY_DATA
+	/** If true clears the patch instead of generating one when the outer DMX entity is invalidated */
+	UPROPERTY(Transient)
+	bool bRequestClearPatch = false;
+#endif // WITH_EDITORONLY_DATA
+
+	/** 
+	 * If true, this entity defines the patch and its fixture type. 
+	 * If false, this entity only follows the patch, but does not update the fixture type.
+	 */
+	UPROPERTY()
+	bool bIsPrimaryPatch = true;
 
-	/** DMX universe id */
-	UPROPERTY(EditAnywhere, Category = Mapping, Meta = (ClampMin = "0", UIMin = "0"))
-	int32 Universe = 1;
+	/** The index of the DMX function to receive */
+	UPROPERTY()
+	int32 FunctionIndex = INDEX_NONE;
 
-	/** Starting channel */
-	UPROPERTY(EditAnywhere, Category = Mapping, Meta = (ClampMin = "1", ClampMax = "512", UIMin = "1", UIMax = "512"))
-	int32 StartingChannel = 1;
+	/** The attribute name of this binding */
+	UPROPERTY()
+	FName AttributeName;
 
 	/**
 	 * Least Significant Byte mode makes the individual bytes (channels) of the function be
 	 * interpreted with the first bytes being the lowest part of the number.
-	 * Most Fixtures use MSB (Most Significant Byte).
 	 */
 	UPROPERTY(EditAnywhere, Category = Mapping)
 	bool bUseLSB = false;
@@ -43,13 +60,39 @@ struct FRemoteControlDMXProtocolEntityExtraSetting
 	UPROPERTY(EditAnywhere, Category = Mapping)
 	EDMXFixtureSignalFormat DataType = EDMXFixtureSignalFormat::E8Bit;
 
+public:
+	// Workaround for clang deprecation warnings for deprecated PackageGuid member in implicit constructors	
+	PRAGMA_DISABLE_DEPRECATION_WARNINGS
+	FRemoteControlDMXProtocolEntityExtraSetting() = default;
+	FRemoteControlDMXProtocolEntityExtraSetting(const FRemoteControlDMXProtocolEntityExtraSetting&) = default;
+	FRemoteControlDMXProtocolEntityExtraSetting(FRemoteControlDMXProtocolEntityExtraSetting&&) = default;
+	FRemoteControlDMXProtocolEntityExtraSetting& operator=(const FRemoteControlDMXProtocolEntityExtraSetting&) = default;
+	FRemoteControlDMXProtocolEntityExtraSetting& operator=(FRemoteControlDMXProtocolEntityExtraSetting&&) = default;
+	PRAGMA_ENABLE_DEPRECATION_WARNINGS
+	
+
+	//////////////////////////
+	// DEPRECATED PROPERTIES
+
+#if WITH_EDITORONLY_DATA
+	/** The DMX universe of this entity. -1 in auto patch mode. */
+	UPROPERTY()
+	int32 Universe_DEPRECATED = -1;
+
+	/** The starting channel of this entity. -1 in auto patch mode. */
+	UPROPERTY()
+	int32 StartingChannel_DEPRECATED = -1;
+
 	/** If set to true, uses the default input port set in Remote Control Protocol project settings */
-	UPROPERTY(EditAnywhere, Category = Mapping)
-	bool bUseDefaultInputPort = true;
+	UE_DEPRECATED(5.5, "Remote control now uses a DMX Library internally. Please refer to the Fixture Patch ref instead.")
+	UPROPERTY()
+	bool bUseDefaultInputPort_DEPRECATED = true;
 
 	/** Reference of an input DMX port id */
-	UPROPERTY(EditAnywhere, Category = Mapping)
-	FGuid InputPortId;
+	UE_DEPRECATED(5.5, "Remote control now uses a DMX Library internally. Please refer to the Fixture Patch ref instead.")
+	UPROPERTY()
+	FGuid InputPortId_DEPRECATED;
+#endif 
 
 };
 
@@ -61,29 +104,38 @@ struct FRemoteControlDMXProtocolEntity : public FRemoteControlProtocolEntity
 {
 	GENERATED_BODY()
 
-	friend class FRemoteControlProtocolDMX;
-
 public:
+	/** Returns the num DMX channels this setting spans */
+	uint8 GetNumDMXChannels() const;
+
 	//~ Begin FRemoteControlProtocolEntity interface
 	virtual FName GetRangePropertyName() const override { return NAME_UInt32Property; }
 	virtual uint8 GetRangePropertySize() const override;
 	virtual const FString& GetRangePropertyMaxValue() const override;
+	//~ End FRemoteControlProtocolEntity interface
 
-#if WITH_EDITOR
+	/** Invalidates this entity. The entity will be updated on the next tick. */
+	REMOTECONTROLPROTOCOLDMX_API void Invalidate();
 
-	/** Register(s) all the widgets of this protocol entity. */
-	virtual void RegisterProperties() override;
+	/** Finds remote control protocol DMX entities used by the specified property */
+	REMOTECONTROLPROTOCOLDMX_API static TArray<TSharedRef<TStructOnScope<FRemoteControlProtocolEntity>>> GetAllDMXProtocolEntitiesInPreset(URemoteControlPreset* Preset);
 
-#endif // WITH_EDITOR
-	//~ End FRemoteControlProtocolEntity interface
+	/** Finds remote control protocol DMX entities used by the specified property */
+	REMOTECONTROLPROTOCOLDMX_API static TArray<TSharedRef<TStructOnScope<FRemoteControlProtocolEntity>>> FindEntitiesByProperty(const TSharedRef<FRemoteControlProperty>& Property);
 
-	/** Initialize struct and delegates */
-	void Initialize();
+	/** Binds this entity to DMX */
+	void BindDMX();
 
-	/** Try to get the port ID from dmx protocol settings */
-	void UpdateInputPort();
+	/** Unbinds this entity from DMX */
+	void UnbindDMX();
 
-	/** Extra protocol settings. Primary using for customization */
+	/** Called when the struct is serialized */
+	bool Serialize(FArchive& Ar);
+
+	/** Called after the struct is serialized */
+	void PostSerialize(const FArchive& Ar);
+
+	/** Extra protocol settings */
 	UPROPERTY(EditAnywhere, Category = Mapping, meta = (ShowOnlyInnerProperties))
 	FRemoteControlDMXProtocolEntityExtraSetting ExtraSetting;
 
@@ -92,24 +144,16 @@ public:
 	uint32 RangeInputTemplate = 0;
 
 private:
-	/** DMX entity cache buffer. From 1 up to 4 channels, based on DataType */
-	TArray<uint8> CacheDMXBuffer;
-
-	/** A single, generic DMX signal. One universe of raw DMX data received */
-	FDMXSignalSharedPtr LastSignalPtr;
-
-public:
-	/** Called when the struct is serialized */
-	bool Serialize(FArchive& Ar);
-
-	/** Called after the struct is serialized */
-	void PostSerialize(const FArchive& Ar);
+	/** Gets or creates the DMX library proxy */
+	static URemoteControlDMXLibraryProxy* GetDMXLibraryProxy(URemoteControlPreset* Preset);
 
 	// DEPRECATED MEMBERS
+public:
+#if WITH_EDITORONLY_DATA
 	// Deprecated 5.0
 	UPROPERTY(Meta = (DeprecatedProperty, DeprecationMessage = "This Property is deprecated and will be removed in a future release. It was moved to the ExtraSetting struct member so the property can be customized."))
-	int32 Universe_DEPRECATED = 1;
-	
+	int32 Universe_DEPRECATED = 0;
+
 	UPROPERTY(Meta = (DeprecatedProperty, DeprecationMessage = "This Property is deprecated and will be removed in a future release. It was moved to the ExtraSetting struct member so the property can be customized."))
 	bool bUseLSB_DEPRECATED = false;
 
@@ -118,10 +162,12 @@ public:
 
 	UPROPERTY(Meta = (DeprecatedProperty, DeprecationMessage = "This Property is deprecated and will be removed in a future release. It was moved to the ExtraSetting struct member so the property can be customized."))
 	bool bUseDefaultInputPort_DEPRECATED = true;
-	
+
 	UPROPERTY(Meta = (DeprecatedProperty, DeprecationMessage = "This Property is deprecated and will be removed in a future release. It was moved to the ExtraSetting struct member so the property can be customized."))
 	FGuid InputPortId_DEPRECATED;
+#endif
 };
+
 template<>
 struct TStructOpsTypeTraits<FRemoteControlDMXProtocolEntity> : public TStructOpsTypeTraitsBase2<FRemoteControlDMXProtocolEntity>
 {
@@ -143,46 +189,28 @@ public:
 	{}
 	
 	//~ Begin IRemoteControlProtocol interface
-	virtual void Bind(FRemoteControlProtocolEntityPtr InRemoteControlProtocolEntityPtr) override;
-	virtual void Unbind(FRemoteControlProtocolEntityPtr InRemoteControlProtocolEntityPtr) override;
+	virtual void Bind(TSharedPtr<TStructOnScope<FRemoteControlProtocolEntity>> InRemoteControlProtocolEntity) override;
+	virtual void Unbind(TSharedPtr<TStructOnScope<FRemoteControlProtocolEntity>> InRemoteControlProtocolEntity) override;
 	virtual void UnbindAll() override;
 	virtual UScriptStruct* GetProtocolScriptStruct() const override { return FRemoteControlDMXProtocolEntity::StaticStruct(); }
-	virtual void OnEndFrame() override;
 	//~ End IRemoteControlProtocol interface
 
-private:
-	/**
-	 * Apply dmx channel data to the bound property, potentially resize the cache buffer
-	 * @param InSignal				DMX signal buffer pointer
-	 * @param InDMXOffset			Byte offset in signal buffer
-	 * @param InProtocolEntityPtr	Protocol entity pointer
-	 */
-	void ProcessAndApplyProtocolValue(const FDMXSignalSharedPtr& InSignal, int32 InDMXOffset, const FRemoteControlProtocolEntityPtr& InProtocolEntityPtr);
+	/** DMX protocol name */
+	REMOTECONTROLPROTOCOLDMX_API static const FName ProtocolName;
 
 #if WITH_EDITOR
-	/**
-	 * Process the AutoBinding to the Remote Control Entity
-	 * @param InProtocolEntityPtr	Protocol entity pointer
-	 */
-	void ProcessAutoBinding(const FRemoteControlProtocolEntityPtr& InProtocolEntityPtr);
+	REMOTECONTROLPROTOCOLDMX_API static const FName PatchColumnName;
+	REMOTECONTROLPROTOCOLDMX_API static const FName UniverseColumnName;
+	REMOTECONTROLPROTOCOLDMX_API static const FName ChannelColumnName;
+#endif // WITH_EDITOR
 
 protected:
-
+#if WITH_EDITOR
 	/** Populates protocol specific columns. */
 	virtual void RegisterColumns() override;
-
-#endif // WITH_EDITOR
+#endif // WITH_EDIOR
 
 private:
 	/** Binding for the DMX protocol */
-	TArray<FRemoteControlProtocolEntityWeakPtr> ProtocolsBindings;
-
-#if WITH_EDITORONLY_DATA
-	/** DMX universe cache buffer.*/
-	TArray<uint8, TFixedAllocator<DMX_UNIVERSE_SIZE>> CacheUniverseDMXBuffer;
-#endif
-
-public:
-	/** DMX protocol name */
-	static const FName ProtocolName;
+	TArray<FRemoteControlProtocolEntityWeakPtr> WeakProtocolsBindings;
 };
diff --git a/Engine/Plugins/VirtualProduction/RemoteControlProtocolDMX/Source/RemoteControlProtocolDMX/Public/RemoteControlProtocolDMXSettings.h b/Engine/Plugins/VirtualProduction/RemoteControlProtocolDMX/Source/RemoteControlProtocolDMX/Public/RemoteControlProtocolDMXSettings.h
--- a/Engine/Plugins/VirtualProduction/RemoteControlProtocolDMX/Source/RemoteControlProtocolDMX/Public/RemoteControlProtocolDMXSettings.h
+++ b/Engine/Plugins/VirtualProduction/RemoteControlProtocolDMX/Source/RemoteControlProtocolDMX/Public/RemoteControlProtocolDMXSettings.h
@@ -2,15 +2,14 @@
 
 #pragma once
 
-#include "CoreMinimal.h"
-
-#include "RemoteControlProtocolDMX.h"
 #include "UObject/Object.h"
 
 #include "RemoteControlProtocolDMXSettings.generated.h"
 
+class UE_DEPRECATED(5.5, "This class currently only contains deprecated properties and is not registered with project settings.") URemoteControlProtocolDMXSettings;
 /**
- * DMX Remote Control Settings
+ * DMX Remote Control Settings.
+ * This class currently only contains deprecated properties and is not registered with project settings.
  */
 UCLASS(Config = Engine, DefaultConfig)
 class REMOTECONTROLPROTOCOLDMX_API URemoteControlProtocolDMXSettings : public UObject
@@ -18,24 +17,22 @@ class REMOTECONTROLPROTOCOLDMX_API URemoteControlProtocolDMXSettings : public UO
 	GENERATED_BODY()
 
 public:
-	//~ Begin UObject interface
-#if WITH_EDITOR
-	virtual void PostEditChangeProperty(struct FPropertyChangedEvent& PropertyChangedEvent) override;
-#endif
-	//~ End UObject interface
-
-	UFUNCTION(BlueprintGetter)
+	UE_DEPRECATED(5.5, "Remote control DMX no longer has an input port, instead please refer to the input port of the DMX library in DMXUserData of the Remote Control Preset")
 	FGuid GetOrCreateDefaultInputPortId();
 
 	/** Returns a delegate broadcast whenever the Remote Control Protocol DMX Settings changed */
+	UE_DEPRECATED(5.5, "Currently URemoteControlProtocolDMXSettings has no properties that should signal a change.")
 	static FSimpleMulticastDelegate& GetOnRemoteControlProtocolDMXSettingsChanged();
 
-	static FName GetDefaultInputPortIdPropertyNameChecked() { return GET_MEMBER_NAME_CHECKED(URemoteControlProtocolDMXSettings, DefaultInputPortId); }
+#if WITH_EDITOR
+	UE_DEPRECATED(5.5, "The DefaultInputPortId property is deprecated and should no longer be accessed")
+	static FName GetDefaultInputPortIdPropertyNameChecked() { return GET_MEMBER_NAME_CHECKED(URemoteControlProtocolDMXSettings, DefaultInputPortId_DEPRECATED); }
+#endif // WITH_EDITOR
 
 private:
 	/** DMX Default Device */
-	UPROPERTY(Config, EditAnywhere, BlueprintReadOnly, Meta = (BlueprintGetter = GetOrCreateDefaultInputPortId, AllowPrivateAccess = true), Category = Mapping)
-	FGuid DefaultInputPortId;
+	UPROPERTY()
+	FGuid DefaultInputPortId_DEPRECATED;
 
-	static FSimpleMulticastDelegate OnRemoteControlProtocolDMXSettingsChangedDelegate;
+	static FSimpleMulticastDelegate OnRemoteControlProtocolDMXSettingsChangedDelegate_DEPRECATED;
 };
diff --git a/Engine/Plugins/VirtualProduction/RemoteControlProtocolDMX/Source/RemoteControlProtocolDMX/RemoteControlProtocolDMX.Build.cs b/Engine/Plugins/VirtualProduction/RemoteControlProtocolDMX/Source/RemoteControlProtocolDMX/RemoteControlProtocolDMX.Build.cs
--- a/Engine/Plugins/VirtualProduction/RemoteControlProtocolDMX/Source/RemoteControlProtocolDMX/RemoteControlProtocolDMX.Build.cs
+++ b/Engine/Plugins/VirtualProduction/RemoteControlProtocolDMX/Source/RemoteControlProtocolDMX/RemoteControlProtocolDMX.Build.cs
@@ -17,19 +17,10 @@ public class RemoteControlProtocolDMX : ModuleRules
 		PrivateDependencyModuleNames.AddRange(
 			new string[] {
 				"DMXProtocol",
+				"Engine",
 				"RemoteControl",
 				"RemoteControlProtocol",
 			}
 		);
-		
-		if (Target.bBuildEditor)
-		{
-			PrivateDependencyModuleNames.AddRange(
-				new string[] {
-					"InputCore",
-					"RemoteControlProtocolWidgets",
-				}
-			);
-		}
 	}
 }
diff --git a/Engine/Plugins/VirtualProduction/RemoteControlProtocolDMX/Source/RemoteControlProtocolDMXEditor/Private/Customizations/RemoteControlProtocolDMXPresetUserDataDetails.cpp b/Engine/Plugins/VirtualProduction/RemoteControlProtocolDMX/Source/RemoteControlProtocolDMXEditor/Private/Customizations/RemoteControlProtocolDMXPresetUserDataDetails.cpp
new file mode 100644
--- /dev/null
+++ b/Engine/Plugins/VirtualProduction/RemoteControlProtocolDMX/Source/RemoteControlProtocolDMXEditor/Private/Customizations/RemoteControlProtocolDMXPresetUserDataDetails.cpp
@@ -0,0 +1,36 @@
+// Copyright Epic Games, Inc. All Rights Reserved.
+
+#include "RemoteControlProtocolDMXPresetUserDataDetails.h"
+
+#include "DetailCategoryBuilder.h"
+#include "DetailLayoutBuilder.h"
+#include "DetailWidgetRow.h"
+#include "RemoteControlDMXUserData.h"
+#include "RemoteControlProtocolDMXSettings.h"
+#include "Styling/StarshipCoreStyle.h"
+#include "Styling/ToolBarStyle.h"
+#include "Widgets/Layout/SWrapBox.h"
+
+#define LOCTEXT_NAMESPACE "RemoteControlProtocolDMXPresetUserDataDetails"
+
+namespace UE::RemoteControl::DMX
+{
+	TSharedRef<IDetailCustomization> FRemoteControlProtocolDMXPresetUserDataDetails::MakeInstance()
+	{
+		return MakeShared<FRemoteControlProtocolDMXPresetUserDataDetails>();
+	}
+
+	void FRemoteControlProtocolDMXPresetUserDataDetails::CustomizeDetails(IDetailLayoutBuilder& DetailBuilder)
+	{
+		DMXLibraryHandle = DetailBuilder.GetProperty(URemoteControlDMXUserData::GetDMXLibraryMemberNameChecked());
+		DMXLibraryHandle->MarkHiddenByCustomization();
+		DMXLibraryHandle->MarkResetToDefaultCustomized();
+
+		constexpr const TCHAR* NoCategoryName = TEXT("nocategory");
+		IDetailCategoryBuilder& CategoryBuilder = DetailBuilder.EditCategory(NoCategoryName);
+
+		CategoryBuilder.AddProperty(DMXLibraryHandle);
+	}
+}
+
+#undef LOCTEXT_NAMESPACE
diff --git a/Engine/Plugins/VirtualProduction/RemoteControlProtocolDMX/Source/RemoteControlProtocolDMXEditor/Private/Customizations/RemoteControlProtocolDMXPresetUserDataDetails.h b/Engine/Plugins/VirtualProduction/RemoteControlProtocolDMX/Source/RemoteControlProtocolDMXEditor/Private/Customizations/RemoteControlProtocolDMXPresetUserDataDetails.h
new file mode 100644
--- /dev/null
+++ b/Engine/Plugins/VirtualProduction/RemoteControlProtocolDMX/Source/RemoteControlProtocolDMXEditor/Private/Customizations/RemoteControlProtocolDMXPresetUserDataDetails.h
@@ -0,0 +1,26 @@
+// Copyright Epic Games, Inc. All Rights Reserved.
+
+#pragma once
+
+#include "IDetailCustomization.h"
+
+class IPropertyHandle;
+
+namespace UE::RemoteControl::DMX
+{
+	/** Details customization for URemoteControlDMXPresetUserData */
+	class FRemoteControlProtocolDMXPresetUserDataDetails final : public IDetailCustomization
+	{
+	public:
+		/** Creates a new instance of this details customization. */
+		static TSharedRef<IDetailCustomization> MakeInstance();
+
+		//~ Begin IDetailCustomization interface
+		virtual void CustomizeDetails(IDetailLayoutBuilder& DetailBuilder) override;
+		//~ End IDetailCustomization interface
+
+	private:
+		/** Property handle for the DMXLibrary property */
+		TSharedPtr<IPropertyHandle> DMXLibraryHandle;
+	};
+}
diff --git a/Engine/Plugins/VirtualProduction/RemoteControlProtocolDMX/Source/RemoteControlProtocolDMXEditor/Private/Library/RemoteControlDMXLibraryBuilder.cpp b/Engine/Plugins/VirtualProduction/RemoteControlProtocolDMX/Source/RemoteControlProtocolDMXEditor/Private/Library/RemoteControlDMXLibraryBuilder.cpp
new file mode 100644
--- /dev/null
+++ b/Engine/Plugins/VirtualProduction/RemoteControlProtocolDMX/Source/RemoteControlProtocolDMXEditor/Private/Library/RemoteControlDMXLibraryBuilder.cpp
@@ -0,0 +1,216 @@
+// Copyright Epic Games, Inc. All Rights Reserved.
+
+#include "RemoteControlDMXLibraryBuilder.h"
+
+#include "Algo/AnyOf.h"
+#include "RemoteControlDMXUserData.h"
+#include "RemoteControlPreset.h"
+#include "Library/DMXEntityFixturePatch.h"
+#include "Library/DMXLibrary.h"
+#include "Library/RemoteControlDMXControlledPropertyPatch.h"
+#include "Library/RemoteControlDMXLibraryProxy.h"
+#include "Library/RemoteControlDMXPatchBuilder.h"
+
+namespace UE::RemoteControl::DMX
+{
+	void FRemoteControlDMXLibraryBuilder::Register()
+	{
+		static TSharedRef<FRemoteControlDMXLibraryBuilder> Instance = MakeShared<FRemoteControlDMXLibraryBuilder>();
+
+		URemoteControlDMXLibraryProxy::GetOnPrePropertyPatchesChanged().AddSP(Instance, &FRemoteControlDMXLibraryBuilder::PrePropertyPatchesChanged);
+		URemoteControlDMXLibraryProxy::GetOnPostPropertyPatchesChanged().AddSP(Instance, &FRemoteControlDMXLibraryBuilder::PostPropertyPatchesChanged);
+	}
+
+	void FRemoteControlDMXLibraryBuilder::PrePropertyPatchesChanged(URemoteControlPreset* InOutPreset)
+	{
+		// Set the preset to work with 
+		Preset = InOutPreset;
+
+		URemoteControlDMXLibraryProxy* DMXLibraryProxy = GetDMXLibraryProxy();
+		UDMXLibrary* DMXLibrary = GetDMXLibrary();
+		if (DMXLibraryProxy && DMXLibrary)
+		{
+			// Remember fixture patches and fixture types before property patches changed
+			const TArray<TSharedRef<FRemoteControlDMXControlledPropertyPatch>> PreEditChangePropertyPatches = DMXLibraryProxy->GetPropertyPatches();
+
+			PreviousFixturePatches.Reset();
+			Algo::TransformIf(PreEditChangePropertyPatches, PreviousFixturePatches,
+				[](const TSharedRef<FRemoteControlDMXControlledPropertyPatch>& PropertyPatch)
+				{
+					return PropertyPatch->GetFixturePatch();
+				},
+				[](const TSharedRef<FRemoteControlDMXControlledPropertyPatch>& PropertyPatch)
+				{
+					return PropertyPatch->GetFixturePatch();
+				});
+
+			PreviousFixtureTypes.Reset();
+			Algo::TransformIf(PreviousFixturePatches, PreviousFixtureTypes,
+				[](const UDMXEntityFixturePatch* FixturePatch)
+				{
+					return FixturePatch->GetFixtureType();
+				},
+				[](const UDMXEntityFixturePatch* FixturePatch)
+				{
+					return FixturePatch->GetFixtureType();
+				});
+		}
+	}
+
+	void FRemoteControlDMXLibraryBuilder::PostPropertyPatchesChanged()
+	{
+		if (URemoteControlDMXLibraryProxy* DMXLibraryProxy = GetDMXLibraryProxy())
+		{
+			const TArray<TSharedRef<FRemoteControlDMXControlledPropertyPatch>> PostEditChangePropertyPatches = DMXLibraryProxy->GetPropertyPatches();
+			FRemoteControlDMXPatchBuilder::BuildFixturePatches(AsShared(), PostEditChangePropertyPatches);
+
+			const URemoteControlDMXUserData* DMXUserData = GetDMXUserData();
+			const bool bAutoAssign = DMXUserData && DMXUserData->IsAutoPatch();
+			if (bAutoAssign)
+			{
+				AutoAssignFixturePatches(PostEditChangePropertyPatches);
+			}
+
+			RemoveObsoleteFixturesFromDMXLibrary(PostEditChangePropertyPatches);
+		}
+	}
+
+	URemoteControlDMXUserData* FRemoteControlDMXLibraryBuilder::GetDMXUserData() const
+	{
+		return URemoteControlDMXUserData::GetOrCreateDMXUserData(Preset);
+	}
+
+	URemoteControlDMXLibraryProxy* FRemoteControlDMXLibraryBuilder::GetDMXLibraryProxy() const
+	{
+		URemoteControlDMXUserData* DMXUserData = GetDMXUserData();
+		return DMXUserData ? DMXUserData->GetDMXLibraryProxy() : nullptr;
+	}
+
+	UDMXLibrary* FRemoteControlDMXLibraryBuilder::GetDMXLibrary() const
+	{
+		URemoteControlDMXLibraryProxy* DMXLibraryProxy = GetDMXLibraryProxy();
+		return DMXLibraryProxy ? DMXLibraryProxy->GetDMXLibrary() : nullptr;
+	}
+
+	void FRemoteControlDMXLibraryBuilder::RemoveObsoleteFixturesFromDMXLibrary(const TArray<TSharedRef<FRemoteControlDMXControlledPropertyPatch>>& PostEditChangePropertyPatches)
+	{
+		UDMXLibrary* DMXLibrary = GetDMXLibrary();
+		if (DMXLibrary)
+		{
+			for (UDMXEntityFixturePatch* PreviousFixturePatch : PreviousFixturePatches)
+			{
+				if (!PreviousFixturePatch || !PreviousFixturePatch->GetParentLibrary())
+				{
+					continue;
+				}
+
+				const bool bFixturePatchStillReferenced = Algo::AnyOf(PostEditChangePropertyPatches,
+					[PreviousFixturePatch](const TSharedRef<FRemoteControlDMXControlledPropertyPatch>& PropertyPatch)
+					{
+						if (PropertyPatch->GetFixturePatch() == PreviousFixturePatch)
+						{
+							return true;
+						}
+
+						return false;
+					});
+
+				if (!bFixturePatchStillReferenced)
+				{
+					UDMXEntityFixturePatch::RemoveFixturePatchFromLibrary(PreviousFixturePatch);
+				}
+			}
+
+			for (UDMXEntityFixtureType* PreviousFixtureType : PreviousFixtureTypes)
+			{
+				if (!PreviousFixtureType || !PreviousFixtureType->GetParentLibrary())
+				{
+					continue;
+				}
+
+				const TArray<UDMXEntityFixturePatch*> FixturePatchesInLibrary = DMXLibrary->GetEntitiesTypeCast<UDMXEntityFixturePatch>();
+
+				const bool bFixtureTypeStillReferenced = Algo::AnyOf(FixturePatchesInLibrary, [PreviousFixtureType](const UDMXEntityFixturePatch* FixturePatch)
+					{
+						return FixturePatch && FixturePatch->GetFixtureType() == PreviousFixtureType;
+					});
+
+				if (!bFixtureTypeStillReferenced)
+				{
+					UDMXEntityFixtureType::RemoveFixtureTypeFromLibrary(PreviousFixtureType);
+				}
+			}
+		}
+	}
+
+	void FRemoteControlDMXLibraryBuilder::AutoAssignFixturePatches(const TArray<TSharedRef<FRemoteControlDMXControlledPropertyPatch>>& PostEditChangePropertyPatches)
+	{
+		URemoteControlDMXUserData* DMXUserData = GetDMXUserData();
+		UDMXLibrary* DMXLibrary = GetDMXLibrary();
+		if (!DMXUserData || !DMXLibrary)
+		{
+			return;
+		}
+
+		TArray<UDMXEntityFixturePatch*> FixturePatches;
+		Algo::TransformIf(PostEditChangePropertyPatches, FixturePatches,
+			[](const TSharedRef<FRemoteControlDMXControlledPropertyPatch>& PropertyPatch)
+			{
+				return PropertyPatch->GetFixturePatch() != nullptr;
+			},
+			[](const TSharedRef<FRemoteControlDMXControlledPropertyPatch>& PropertyPatch)
+			{
+				return PropertyPatch->GetFixturePatch();
+			});
+
+		// Reset patch
+		for (UDMXEntityFixturePatch* FixturePatch : FixturePatches)
+		{
+			FixturePatch->SetStartingChannel(1);
+			FixturePatch->SetUniverseID(1);
+		}
+
+		// Auto assign
+		for (const TSharedRef<FRemoteControlDMXControlledPropertyPatch>& PropertyPatch : PostEditChangePropertyPatches)
+		{
+			UDMXEntityFixturePatch* FixturePatch = PropertyPatch->GetFixturePatch();
+			if (!FixturePatch)
+			{
+				continue;
+			}
+
+			const int32 FixturePatchIndex = FixturePatches.IndexOfByKey(FixturePatch);
+			if (!ensureMsgf(FixturePatchIndex != INDEX_NONE, TEXT("Unexpected cannot find fixture patch in DMX Library. Cannot auto assign fixture patch.")))
+			{
+				continue;
+			}
+
+			const int32 AutoAssingFromUniverse = DMXUserData->GetAutoAssignFromUniverse();
+			const UDMXEntityFixturePatch* PreviousPatch = FixturePatches.IsValidIndex(FixturePatchIndex - 1) ? FixturePatches[FixturePatchIndex - 1] : nullptr;
+			const int64 DesiredAbsoluteStartingChannel = [PreviousPatch, AutoAssingFromUniverse]()
+				{
+					if (PreviousPatch && PreviousPatch->GetUniverseID() >= AutoAssingFromUniverse)
+					{
+						return (int64)PreviousPatch->GetUniverseID() * DMX_UNIVERSE_SIZE + PreviousPatch->GetEndingChannel();
+					}
+					else
+					{
+						return (int64)AutoAssingFromUniverse * DMX_UNIVERSE_SIZE;
+					}
+				}();
+
+			const bool bFitsUniverse = DesiredAbsoluteStartingChannel % DMX_UNIVERSE_SIZE + FixturePatch->GetChannelSpan() <= DMX_UNIVERSE_SIZE;
+			const int64 AbsoluteStartingChannel = bFitsUniverse ? DesiredAbsoluteStartingChannel : (DesiredAbsoluteStartingChannel / DMX_UNIVERSE_SIZE + 1) * DMX_UNIVERSE_SIZE;
+
+			const int32 Universe = AbsoluteStartingChannel / DMX_UNIVERSE_SIZE;
+			const int32 Channel = AbsoluteStartingChannel % DMX_UNIVERSE_SIZE + 1;
+
+			FixturePatch->PreEditChange(nullptr);
+
+			FixturePatch->SetUniverseID(Universe);
+			FixturePatch->SetStartingChannel(Channel);
+
+			FixturePatch->PostEditChange();
+		}
+	}
+}
diff --git a/Engine/Plugins/VirtualProduction/RemoteControlProtocolDMX/Source/RemoteControlProtocolDMXEditor/Private/Library/RemoteControlDMXLibraryBuilder.h b/Engine/Plugins/VirtualProduction/RemoteControlProtocolDMX/Source/RemoteControlProtocolDMXEditor/Private/Library/RemoteControlDMXLibraryBuilder.h
new file mode 100644
--- /dev/null
+++ b/Engine/Plugins/VirtualProduction/RemoteControlProtocolDMX/Source/RemoteControlProtocolDMXEditor/Private/Library/RemoteControlDMXLibraryBuilder.h
@@ -0,0 +1,54 @@
+// Copyright Epic Games, Inc. All Rights Reserved.
+
+#pragma once
+
+#include "Library/RemoteControlDMXLibraryProxy.h"
+#include "Templates/SharedPointer.h"
+
+class UDMXEntityFixtureType;
+class UDMXLibrary;
+class URemoteControlDMXLibraryProxy;
+class URemoteControlPreset;
+
+namespace UE::RemoteControl::DMX
+{
+	/** Builds a DMX Library for a Remote Control Preset */
+	class FRemoteControlDMXLibraryBuilder
+		: public TSharedFromThis<FRemoteControlDMXLibraryBuilder>
+	{
+	public:
+		/** Creates an instance of the DMX Library builder */
+		static void Register();
+
+		/** Returns the DMX User Data of the Remote Control Preset */
+		URemoteControlDMXUserData* GetDMXUserData() const;
+
+		/** Returns the DMX Library Proxy of the Remote Control Preset */
+		URemoteControlDMXLibraryProxy* GetDMXLibraryProxy() const;
+
+		/** Returns the DMX Library of used with this Remote Control Preset */
+		UDMXLibrary* GetDMXLibrary() const;
+
+	private:
+		/** Called before property patches are being changed */
+		void PrePropertyPatchesChanged(URemoteControlPreset* InOutPreset);
+
+		/** Called after property patches were changed */
+		void PostPropertyPatchesChanged();
+
+		/** Removes unused Fixture Types and Fixture Patches from the DMX Library */
+		void RemoveObsoleteFixturesFromDMXLibrary(const TArray<TSharedRef<FRemoteControlDMXControlledPropertyPatch>>& PostEditChangePropertyPatches);
+
+		/** Auto assigns fixture patches in the DMX Library */
+		void AutoAssignFixturePatches(const TArray<TSharedRef<FRemoteControlDMXControlledPropertyPatch>>& PostEditChangePropertyPatches);
+
+		/** Fixture types used before property patches changed */
+		TArray<UDMXEntityFixtureType*> PreviousFixtureTypes;
+
+		/** Fixture patches used before property patches changed */
+		TArray<UDMXEntityFixturePatch*> PreviousFixturePatches;
+
+		/** The remote control preset for which the Library is built */
+		URemoteControlPreset* Preset = nullptr;
+	};
+}
diff --git a/Engine/Plugins/VirtualProduction/RemoteControlProtocolDMX/Source/RemoteControlProtocolDMXEditor/Private/Library/RemoteControlDMXPatchBuilder.cpp b/Engine/Plugins/VirtualProduction/RemoteControlProtocolDMX/Source/RemoteControlProtocolDMXEditor/Private/Library/RemoteControlDMXPatchBuilder.cpp
new file mode 100644
--- /dev/null
+++ b/Engine/Plugins/VirtualProduction/RemoteControlProtocolDMX/Source/RemoteControlProtocolDMXEditor/Private/Library/RemoteControlDMXPatchBuilder.cpp
@@ -0,0 +1,539 @@
+// Copyright Epic Games, Inc. All Rights Reserved.
+
+#include "RemoteControlDMXPatchBuilder.h"
+
+#include "Algo/AnyOf.h"
+#include "Algo/Find.h"
+#include "Algo/MaxElement.h"
+#include "GameFramework/Actor.h"
+#include "Library/DMXEntityFixturePatch.h"
+#include "Library/DMXLibrary.h"
+#include "Library/RemoteControlDMXControlledProperty.h"
+#include "Library/RemoteControlDMXControlledPropertyPatch.h"
+#include "Library/RemoteControlDMXLibraryBuilder.h"
+#include "RemoteControlDMXUserData.h"
+#include "RemoteControlField.h"
+#include "RemoteControlProtocolDMX.h"
+
+#define LOCTEXT_NAMESPACE "FRemoteControlDMXPatchBuilder"
+
+namespace UE::RemoteControl::DMX
+{
+	namespace Internal
+	{
+		/** Struct holding data that needs to be transported from an old to a new patch */
+		struct FRCDMXFixturePatchData
+		{
+			FGuid MVRFixtureUUID;
+			int32 FixtureID = 1;
+		};
+
+		/** Builds a single fixture patch and related fixture type if required from a property patch. */
+		class FRCSinglePatchBuilder
+			: FNoncopyable
+		{
+		public:
+			/** Constructs a single fixture patch and related fixture type if required from a property patch. */
+			FRCSinglePatchBuilder(
+				const TSharedRef<FRemoteControlDMXLibraryBuilder>& InDMXLibraryBuilder,
+				const TSharedRef<FRemoteControlDMXControlledPropertyPatch>& InDMXControlledPropertyPatch);
+
+			/** Rebuilds a fixture patch in the DMX Library */
+			void RebuildFixturePatch();
+
+			/** Returns true if this builder builds a primary patch */
+			bool IsPrimary() const;
+
+			/** Returns true if the patch should be cleared */
+			bool ShouldClearPatch() const;
+
+			/** Caches fixture patch data into the FixturePatchData member */
+			void CacheFixturePatchData();
+
+			/** Clears all fixture types used by the entities of this patch */
+			void ClearFixtureTypes();
+
+			/** Clears all fixture patches  by the entities of this patch */
+			void ClearFixturePatches();
+
+			/** Updates a primary fixture patch */
+			void UpdatePrimaryFixturePatch();
+
+			/** Updates a secondary fixture patch */
+			void UpdateSecondaryFixturePatch();
+
+			/** Finds a fixture type that corresponds with this property patch */
+			UDMXEntityFixtureType* FindFixtureType(UDMXLibrary& DMXLibrary);
+
+			/** Updates an existing fixture type that corresponds with this property patch */
+			void UpdateFixtureType(UDMXLibrary& InDMXLibrary, UDMXEntityFixtureType*& InOutFixtureType);
+
+			/** Finds a fixture patch that corresponds with this property patch */
+			UDMXEntityFixturePatch* FindFixturePatch(UDMXEntityFixtureType* FixtureType) const;
+
+			/** Updates an existing fixture patch that corresponds with this property patch */
+			void UpdateFixturePatch(UDMXLibrary& DMXLibrary, UDMXEntityFixtureType* FixtureType, UDMXEntityFixturePatch* ReuseFixturePatch);
+
+			/** Sets the fixture patch for this property patch. */
+			void SetFixturePatch(UDMXEntityFixturePatch* PrimaryFixturePatch);
+
+			/** Sets if the patch is a primary fixture patch */
+			void SetIsPrimaryFixturePatch(bool bIsPrimaryFixturePatch);
+
+			/** Returns the fixture patches currently stored in DMXControlledProperties. Useful when rebuilding the patch from new properties. */
+			TArray<UDMXEntityFixturePatch*> GetFixturePatchesFromProperties() const;
+
+			/** Returns the desired name for the patch */
+			FString GetDesiredName() const;
+
+			/** Returns the name of the owner object */
+			FString GetOwnerObjectName() const;
+
+			/** Cached data of the fixture patch to use if a patch needs to be regenerated */
+			FRCDMXFixturePatchData FixturePatchData;
+
+			/** The DMX controlled properties contained in the patch */
+			TArray<TSharedRef<FRemoteControlDMXControlledProperty>> DMXControlledProperties;
+
+			/** The library builder for which the patch is created */
+			const TSharedRef<FRemoteControlDMXLibraryBuilder> DMXLibraryBuilder;
+
+			/** The DMX controlled property patch for which a fixture patch should be created */
+			const TSharedRef<FRemoteControlDMXControlledPropertyPatch> DMXControlledPropertyPatch;
+		};
+
+
+		FRCSinglePatchBuilder::FRCSinglePatchBuilder(
+			const TSharedRef<FRemoteControlDMXLibraryBuilder>& InDMXLibraryBuilder,
+			const TSharedRef<FRemoteControlDMXControlledPropertyPatch>& InDMXControlledPropertyPatch)
+			: DMXLibraryBuilder(InDMXLibraryBuilder)
+			, DMXControlledPropertyPatch(InDMXControlledPropertyPatch)
+		{
+			DMXControlledProperties = DMXControlledPropertyPatch->GetDMXControlledProperties();
+
+			CacheFixturePatchData();
+		}
+
+		void FRCSinglePatchBuilder::RebuildFixturePatch()
+		{
+			if (ShouldClearPatch())
+			{
+				ClearFixturePatches();
+				ClearFixtureTypes();
+			}
+			else if (IsPrimary())
+			{
+				UpdatePrimaryFixturePatch();
+			}
+			else
+			{
+				UpdateSecondaryFixturePatch();
+			}
+		}
+
+		bool FRCSinglePatchBuilder::IsPrimary() const
+		{
+			return Algo::AnyOf(DMXControlledProperties, [](const TSharedRef<FRemoteControlDMXControlledProperty>& Property)
+				{
+					const TArray<TSharedRef<TStructOnScope<FRemoteControlProtocolEntity>>>& Entities = Property->GetEntities();
+					const TSharedPtr<TStructOnScope<FRemoteControlProtocolEntity>>& FirstEntity = !Entities.IsEmpty() ? Entities[0].ToSharedPtr() : nullptr;
+					const FRemoteControlDMXProtocolEntity* FirstDMXEntity = FirstEntity.IsValid() && FirstEntity->IsValid() ? Entities[0]->Cast<FRemoteControlDMXProtocolEntity>() : nullptr;
+
+					return FirstDMXEntity ? FirstDMXEntity->ExtraSetting.bIsPrimaryPatch : true;
+				});
+		}
+
+		bool FRCSinglePatchBuilder::ShouldClearPatch() const
+		{
+			return Algo::AnyOf(DMXControlledProperties, [](const TSharedRef<FRemoteControlDMXControlledProperty>& Property)
+				{
+					const TArray<TSharedRef<TStructOnScope<FRemoteControlProtocolEntity>>>& Entities = Property->GetEntities();
+					const TSharedPtr<TStructOnScope<FRemoteControlProtocolEntity>>& FirstEntity = !Entities.IsEmpty() ? Entities[0].ToSharedPtr() : nullptr;
+					const FRemoteControlDMXProtocolEntity* FirstDMXEntity = FirstEntity.IsValid() && FirstEntity->IsValid() ? Entities[0]->Cast<FRemoteControlDMXProtocolEntity>() : nullptr;
+
+					return FirstDMXEntity ? FirstDMXEntity->ExtraSetting.bRequestClearPatch : false;
+				});
+		}
+
+		void FRCSinglePatchBuilder::UpdatePrimaryFixturePatch()
+		{
+			URemoteControlDMXUserData* DMXUserData = DMXLibraryBuilder->GetDMXUserData();
+			URemoteControlDMXLibraryProxy* DMXLibraryProxy = DMXLibraryBuilder->GetDMXLibraryProxy();
+			UDMXLibrary* DMXLibrary = DMXLibraryBuilder->GetDMXLibrary();
+			if (DMXUserData && DMXLibraryProxy && DMXLibrary)
+			{
+				// Make sure all entities are set to be a primary patch
+				SetIsPrimaryFixturePatch(true);
+
+				// Get or create a fixture type
+				UDMXEntityFixtureType* FixtureType = FindFixtureType(*DMXLibrary);
+				UpdateFixtureType(*DMXLibrary, FixtureType);
+
+				// Get or create a fixture patch
+				UDMXEntityFixturePatch* ReuseFixturePatch = FindFixturePatch(FixtureType);
+				UpdateFixturePatch(*DMXLibrary, FixtureType, ReuseFixturePatch);
+			}
+		}
+
+		void FRCSinglePatchBuilder::UpdateSecondaryFixturePatch()
+		{
+			UDMXLibrary* DMXLibrary = DMXLibraryBuilder->GetDMXLibrary();
+			if (!DMXLibrary)
+			{
+				return;
+			}
+
+			// Find the primary fixture patch by MVR Fixture UUID
+			const TArray<UDMXEntityFixturePatch*> FixturePatches = DMXLibrary->GetEntitiesTypeCast<UDMXEntityFixturePatch>();
+			UDMXEntityFixturePatch* const* PrimaryFixturePatchPtr = Algo::FindBy(FixturePatches, FixturePatchData.MVRFixtureUUID, &UDMXEntityFixturePatch::GetMVRFixtureUUID);
+
+			UDMXEntityFixturePatch* PrimaryFixturePatch = PrimaryFixturePatchPtr ? *PrimaryFixturePatchPtr : nullptr;
+
+			const FDMXFixtureMode* ModePtr = PrimaryFixturePatch ? PrimaryFixturePatch->GetActiveMode() : nullptr;
+
+			// Try to follow the primary or make it a primary
+			const bool bCanFollowPrimaryPatch = PrimaryFixturePatch && PrimaryFixturePatch->GetFixtureType() == FindFixtureType(*DMXLibrary);
+			if (bCanFollowPrimaryPatch)
+			{
+				SetIsPrimaryFixturePatch(false);
+				SetFixturePatch(PrimaryFixturePatch);
+			}
+			else
+			{
+				// Change to a primary instead
+				SetIsPrimaryFixturePatch(true);
+				UpdatePrimaryFixturePatch();
+			}
+		}
+
+		void FRCSinglePatchBuilder::CacheFixturePatchData()
+		{
+			if (const UDMXEntityFixturePatch* FixturePatch = DMXControlledPropertyPatch->GetFixturePatch())
+			{
+				FixturePatchData.MVRFixtureUUID = FixturePatch->GetMVRFixtureUUID();
+				FixturePatchData.FixtureID = FixturePatch->GetFixtureID();
+			}
+			else
+			{
+				FixturePatchData.MVRFixtureUUID = FGuid::NewGuid();
+				FixturePatchData.FixtureID = 1;
+			}
+		}
+
+		void FRCSinglePatchBuilder::ClearFixtureTypes()
+		{
+			for (const TSharedRef<FRemoteControlDMXControlledProperty>& Property : DMXControlledProperties)
+			{
+				for (const TSharedRef<TStructOnScope<FRemoteControlProtocolEntity>>& Entity : Property->GetEntities())
+				{
+					FRemoteControlDMXProtocolEntity* DMXEntity = Entity->IsValid() ? Entity->Cast<FRemoteControlDMXProtocolEntity>() : nullptr;
+					UDMXEntityFixturePatch* FixturePatch = DMXEntity ? DMXEntity->ExtraSetting.FixturePatchReference.GetFixturePatch() : nullptr;
+					UDMXEntityFixtureType* FixtureType = FixturePatch ? FixturePatch->GetFixtureType() : nullptr;
+					if (FixturePatch)
+					{
+						UDMXEntityFixtureType::RemoveFixtureTypeFromLibrary(FixtureType);
+					}
+				}
+			}
+		}
+
+		void FRCSinglePatchBuilder::ClearFixturePatches()
+		{
+			for (const TSharedRef<FRemoteControlDMXControlledProperty>& Property : DMXControlledProperties)
+			{
+				for (const TSharedRef<TStructOnScope<FRemoteControlProtocolEntity>>& Entity : Property->GetEntities())
+				{
+					FRemoteControlDMXProtocolEntity* DMXEntity = Entity->IsValid() ? Entity->Cast<FRemoteControlDMXProtocolEntity>() : nullptr;
+					UDMXEntityFixturePatch* FixturePatch = DMXEntity ? DMXEntity->ExtraSetting.FixturePatchReference.GetFixturePatch() : nullptr;
+					if (FixturePatch)
+					{
+						UDMXEntityFixturePatch::RemoveFixturePatchFromLibrary(FixturePatch);
+
+						DMXEntity->ExtraSetting.FixturePatchReference = nullptr;
+					}
+				}
+			}
+		}
+
+		UDMXEntityFixtureType* FRCSinglePatchBuilder::FindFixtureType(UDMXLibrary& DMXLibrary)
+		{
+			TArray<const FRemoteControlDMXProtocolEntity*> DMXEntities;
+			for (const TSharedRef<FRemoteControlDMXControlledProperty>& Property : DMXControlledProperties)
+			{
+				for (const TSharedRef<TStructOnScope<FRemoteControlProtocolEntity>>& Entity : Property->GetEntities())
+				{
+					FRemoteControlDMXProtocolEntity* DMXEntity = Entity->IsValid() ? Entity->Cast<FRemoteControlDMXProtocolEntity>() : nullptr;
+					if (DMXEntity)
+					{
+						DMXEntities.Add(DMXEntity);
+					}
+				}
+			}
+
+			const TArray<UDMXEntityFixtureType*> FixtureTypes = DMXLibrary.GetEntitiesTypeCast<UDMXEntityFixtureType>();
+			UDMXEntityFixtureType* const* FixtureTypePtr = Algo::FindByPredicate(FixtureTypes, [&DMXEntities](const UDMXEntityFixtureType* FixtureType)
+				{
+					if (!FixtureType || FixtureType->Modes.IsEmpty())
+					{
+						return false;
+					}
+					const FDMXFixtureMode& Mode = FixtureType->Modes[0];
+
+					if (Mode.Functions.Num() != DMXEntities.Num())
+					{
+						return false;
+					}
+
+					for (int32 FunctionIndex = 0; FunctionIndex < Mode.Functions.Num(); FunctionIndex++)
+					{
+						const FDMXFixtureFunction& Function = Mode.Functions[FunctionIndex];
+						const FRemoteControlDMXProtocolEntity* DMXEntity = DMXEntities.IsValidIndex(FunctionIndex) ? DMXEntities[FunctionIndex] : nullptr;
+
+						if (DMXEntity &&
+							DMXEntity->ExtraSetting.AttributeName == Function.Attribute.Name &&
+							DMXEntity->ExtraSetting.DataType == Function.DataType &&
+							DMXEntity->ExtraSetting.bUseLSB == Function.bUseLSBMode)
+						{
+							return true;
+						}
+					}
+
+					return false;
+				});
+
+			return FixtureTypePtr ? *FixtureTypePtr : nullptr;
+		}
+
+		void FRCSinglePatchBuilder::UpdateFixtureType(UDMXLibrary& InDMXLibrary, UDMXEntityFixtureType*& InOutFixtureType)
+		{
+			if (!InOutFixtureType)
+			{
+				// Create a new fixture type
+				FDMXEntityFixtureTypeConstructionParams FixtureTypeConstructionParams;
+				FixtureTypeConstructionParams.DMXCategory = FDMXFixtureCategory(TEXT("Remote Control"));
+				FixtureTypeConstructionParams.ParentDMXLibrary = &InDMXLibrary;
+
+				InOutFixtureType = UDMXEntityFixtureType::CreateFixtureTypeInLibrary(FixtureTypeConstructionParams);
+				check(InOutFixtureType);
+
+				const FString DesiredFixtureTypeName = TEXT("FT_") + GetDesiredName();
+				InOutFixtureType->Name = DesiredFixtureTypeName;
+			}
+
+			if (!InOutFixtureType)
+			{
+				return;
+			}
+
+			// Rebuild modes
+			InOutFixtureType->Modes.Reset();
+			InOutFixtureType->Modes.Add(FDMXFixtureMode());
+
+			FDMXFixtureMode& Mode = InOutFixtureType->Modes.Last();
+			Mode.ModeName = TEXT("RemoteControl");
+
+			for (const TSharedRef<FRemoteControlDMXControlledProperty>& Property : DMXControlledProperties)
+			{
+				for (const TSharedRef<TStructOnScope<FRemoteControlProtocolEntity>>& Entity : Property->GetEntities())
+				{
+					FRemoteControlDMXProtocolEntity* DMXEntity = Entity->IsValid() ? Entity->Cast<FRemoteControlDMXProtocolEntity>() : nullptr;
+					if (!DMXEntity)
+					{
+						continue;
+					}
+
+					const int32 NextFreeChannel = [&Mode]()
+						{
+							const FDMXFixtureFunction* MaxFunctionPtr = Algo::MaxElementBy(Mode.Functions, &FDMXFixtureFunction::Channel);
+
+							return MaxFunctionPtr ? MaxFunctionPtr->GetLastChannel() + 1 : 1;
+						}();
+
+						FDMXFixtureFunction NewFunction;
+						NewFunction.FunctionName = Property->ExposedProperty->FieldPathInfo.ToString();
+						NewFunction.Attribute = DMXEntity->ExtraSetting.AttributeName;
+						NewFunction.Channel = NextFreeChannel;
+						NewFunction.DataType = DMXEntity->ExtraSetting.DataType;
+						NewFunction.bUseLSBMode = DMXEntity->ExtraSetting.bUseLSB;
+
+						// Addopt the attribute name from the DMX entity
+						NewFunction.Attribute = DMXEntity->ExtraSetting.AttributeName;
+
+						const int32 FunctionIndex = Mode.Functions.Add(NewFunction);
+
+						// Remember the function index from the fixture function
+						DMXEntity->ExtraSetting.FunctionIndex = FunctionIndex;
+				}
+			}
+		}
+
+		UDMXEntityFixturePatch* FRCSinglePatchBuilder::FindFixturePatch(UDMXEntityFixtureType* FixtureType) const
+		{
+			const TArray<UDMXEntityFixturePatch*> PreviousFixturePatches = GetFixturePatchesFromProperties();
+			if (PreviousFixturePatches.Num() == 1)
+			{
+				const bool bSameFixtureType =
+					PreviousFixturePatches[0] &&
+					PreviousFixturePatches[0]->GetFixtureType() &&
+					PreviousFixturePatches[0]->GetFixtureType() == FixtureType;
+
+				return bSameFixtureType ? PreviousFixturePatches[0] : nullptr;
+			}
+
+			return nullptr;
+		}
+
+		void FRCSinglePatchBuilder::UpdateFixturePatch(UDMXLibrary& DMXLibrary, UDMXEntityFixtureType* FixtureType, UDMXEntityFixturePatch* ReuseFixturePatch)
+		{
+			if (!FixtureType || !ReuseFixturePatch)
+			{
+				ClearFixturePatches();
+			}
+
+			UDMXEntityFixturePatch* FixturePatch = ReuseFixturePatch;
+			if (!FixturePatch && FixtureType)
+			{
+				const TArray<UDMXEntityFixturePatch*> FixturePatches = DMXLibrary.GetEntitiesTypeCast<UDMXEntityFixturePatch>();
+				const bool bMVRFixtureUUIDAlreadyUsed = Algo::AnyOf(FixturePatches, [this](const UDMXEntityFixturePatch* OtherFixturePatch)
+					{
+						return
+							OtherFixturePatch &&
+							OtherFixturePatch->GetMVRFixtureUUID() == FixturePatchData.MVRFixtureUUID;
+					});
+
+				const FGuid UniqueMVRFixtureUUID = bMVRFixtureUUIDAlreadyUsed ? FGuid::NewGuid() : FixturePatchData.MVRFixtureUUID;
+
+				// Create a new fixture patch
+				FDMXEntityFixturePatchConstructionParams FixturePatchConstructionParams;
+				FixturePatchConstructionParams.FixtureTypeRef = FDMXEntityFixtureTypeRef(FixtureType);
+				FixturePatchConstructionParams.ActiveMode = 0;
+				FixturePatchConstructionParams.MVRFixtureUUID = UniqueMVRFixtureUUID;
+
+				FixturePatch = UDMXEntityFixturePatch::CreateFixturePatchInLibrary(FixturePatchConstructionParams);
+				check(FixturePatch);
+
+				FixturePatch->Name = GetDesiredName();
+			}
+
+			if (FixturePatch && FixtureType)
+			{
+				FixturePatch->SetFixtureType(FixtureType);
+				FixturePatch->GenerateFixtureID(FixturePatchData.FixtureID);
+
+				SetFixturePatch(FixturePatch);
+			}
+		}
+
+		void FRCSinglePatchBuilder::SetFixturePatch(UDMXEntityFixturePatch* PrimaryFixturePatch)
+		{
+			for (const TSharedRef<FRemoteControlDMXControlledProperty>& Property : DMXControlledProperties)
+			{
+				for (const TSharedRef<TStructOnScope<FRemoteControlProtocolEntity>>& Entity : Property->GetEntities())
+				{
+					FRemoteControlDMXProtocolEntity* DMXEntity = Entity->IsValid() ? Entity->Cast<FRemoteControlDMXProtocolEntity>() : nullptr;
+					if (DMXEntity)
+					{
+						DMXEntity->ExtraSetting.FixturePatchReference = PrimaryFixturePatch;
+					}
+				}
+			}
+		}
+
+		void FRCSinglePatchBuilder::SetIsPrimaryFixturePatch(bool bIsPrimaryFixturePatch)
+		{
+			for (const TSharedRef<FRemoteControlDMXControlledProperty>& Property : DMXControlledProperties)
+			{
+				for (const TSharedRef<TStructOnScope<FRemoteControlProtocolEntity>>& Entity : Property->GetEntities())
+				{
+					FRemoteControlDMXProtocolEntity* DMXEntity = Entity->IsValid() ? Entity->Cast<FRemoteControlDMXProtocolEntity>() : nullptr;
+					if (DMXEntity)
+					{
+						DMXEntity->ExtraSetting.bIsPrimaryPatch = bIsPrimaryFixturePatch;
+					}
+				}
+			}
+		}
+
+		TArray<UDMXEntityFixturePatch*> FRCSinglePatchBuilder::GetFixturePatchesFromProperties() const
+		{
+			TArray<UDMXEntityFixturePatch*> FixturePatches;
+			for (const TSharedRef<FRemoteControlDMXControlledProperty>& Property : DMXControlledProperties)
+			{
+				if (UDMXEntityFixturePatch* FixturePatch = Property->GetFixturePatch())
+				{
+					FixturePatches.AddUnique(FixturePatch);
+				}
+			}
+
+			return FixturePatches;
+		}
+
+		FString FRCSinglePatchBuilder::GetDesiredName() const
+		{
+			URemoteControlDMXUserData* DMXUserData = DMXLibraryBuilder->GetDMXUserData();
+			if (DMXUserData)
+			{
+				ERemoteControlDMXPatchGroupMode PatchGroupMode = DMXUserData->GetPatchGroupMode();
+
+				if (DMXControlledProperties.IsEmpty())
+				{
+					return TEXT("EmptyPatch");
+				}
+				else if (PatchGroupMode == ERemoteControlDMXPatchGroupMode::GroupByOwner)
+				{
+					return GetOwnerObjectName();
+				}
+				else if (!DMXControlledProperties.IsEmpty())
+				{
+					return FString::Printf(TEXT("%s_%s"), *GetOwnerObjectName(), *DMXControlledProperties[0]->ExposedProperty->FieldName.ToString());
+				}
+			}
+
+			return FText(LOCTEXT("InvalidPatchName", "InvalidPatch")).ToString();
+		}
+
+		FString FRCSinglePatchBuilder::GetOwnerObjectName() const
+		{
+			const UObject* OwnerObject = DMXControlledProperties.IsEmpty() ? nullptr : DMXControlledProperties[0]->GetOwnerObject();
+			if (!OwnerObject)
+			{
+				return FText(LOCTEXT("ObjectNotLoadedInfo", "Object is not loaded")).ToString();
+			}
+
+			if (const AActor* Actor = Cast<AActor>(OwnerObject))
+			{
+				return *Actor->GetActorLabel();
+			}
+			else
+			{
+				return OwnerObject->GetName();
+			}
+		}
+	}
+
+
+	void FRemoteControlDMXPatchBuilder::BuildFixturePatches(
+		const TSharedRef<FRemoteControlDMXLibraryBuilder>& InDMXLibraryBuilder,
+		const TArray<TSharedRef<FRemoteControlDMXControlledPropertyPatch>>& InDMXControlledPropertyPatches)
+	{
+		using FRCSinglePatchBuilder = Internal::FRCSinglePatchBuilder;
+
+		// Create patch builder instances for all property patches
+		TArray<TSharedRef<FRCSinglePatchBuilder>> Builders;
+		Algo::Transform(InDMXControlledPropertyPatches, Builders,
+			[&InDMXLibraryBuilder](const TSharedRef<FRemoteControlDMXControlledPropertyPatch>& PropertyPatch)
+			{
+				return MakeShareable(new FRCSinglePatchBuilder(InDMXLibraryBuilder, PropertyPatch));
+			});
+
+		Algo::StableSortBy(Builders, &FRCSinglePatchBuilder::IsPrimary);
+
+		// Update primary fixture patches first, so secondary fixture patches can follow newly created ones
+		for (const TSharedRef<FRCSinglePatchBuilder>& Builder : Builders)
+		{
+			Builder->RebuildFixturePatch();
+		}
+	}
+}
+
+#undef LOCTEXT_NAMESPACE
diff --git a/Engine/Plugins/VirtualProduction/RemoteControlProtocolDMX/Source/RemoteControlProtocolDMXEditor/Private/Library/RemoteControlDMXPatchBuilder.h b/Engine/Plugins/VirtualProduction/RemoteControlProtocolDMX/Source/RemoteControlProtocolDMXEditor/Private/Library/RemoteControlDMXPatchBuilder.h
new file mode 100644
--- /dev/null
+++ b/Engine/Plugins/VirtualProduction/RemoteControlProtocolDMX/Source/RemoteControlProtocolDMXEditor/Private/Library/RemoteControlDMXPatchBuilder.h
@@ -0,0 +1,21 @@
+// Copyright Epic Games, Inc. All Rights Reserved.
+
+#pragma once
+
+#include "Templates/SharedPointer.h"
+
+namespace UE::RemoteControl::DMX
+{
+	class FRemoteControlDMXControlledPropertyPatch;
+	class FRemoteControlDMXLibraryBuilder;
+
+	/** Builds fixture patches for a DMX Library Builder */
+	class FRemoteControlDMXPatchBuilder
+	{
+	public:
+		/** Builds fixture patches from property patches. Updates both the DMX Library and Remote Control Protocol DMX Entities. */
+		static void BuildFixturePatches(
+			const TSharedRef<FRemoteControlDMXLibraryBuilder>& InDMXLibraryBuilder,
+			const TArray<TSharedRef<FRemoteControlDMXControlledPropertyPatch>>& InDMXControlledPropertyPatches);
+	};
+}
diff --git a/Engine/Plugins/VirtualProduction/RemoteControlProtocolDMX/Source/RemoteControlProtocolDMXEditor/Private/RemoteControlDMXAutoBindHandler.cpp b/Engine/Plugins/VirtualProduction/RemoteControlProtocolDMX/Source/RemoteControlProtocolDMXEditor/Private/RemoteControlDMXAutoBindHandler.cpp
new file mode 100644
--- /dev/null
+++ b/Engine/Plugins/VirtualProduction/RemoteControlProtocolDMX/Source/RemoteControlProtocolDMXEditor/Private/RemoteControlDMXAutoBindHandler.cpp
@@ -0,0 +1,182 @@
+// Copyright Epic Games, Inc. All Rights Reserved.
+
+#include "RemoteControlDMXAutoBindHandler.h"
+
+#include "DMXProtocolTypes.h"
+#include "IO/DMXInputPort.h"
+#include "IRCProtocolBindingList.h"
+#include "IRemoteControlProtocolWidgetsModule.h"
+#include "Library/DMXEntityFixturePatch.h"
+#include "Library/DMXLibrary.h"
+#include "RemoteControlDMXUserData.h"
+#include "RemoteControlField.h"
+#include "RemoteControlPreset.h"
+#include "RemoteControlProtocolDMX.h"
+#include "IRemoteControlProtocolModule.h"
+
+namespace UE::RemoteControl::DMX
+{
+	void FRemoteControlDMXAutoBindHandler::Register()
+	{
+		static const FRemoteControlDMXAutoBindHandler Instance = FRemoteControlDMXAutoBindHandler();
+	}
+
+	void FRemoteControlDMXAutoBindHandler::Tick(float DeltaTime)
+	{
+		if (!UObjectInitialized())
+		{
+			return;
+		}
+
+		const IRemoteControlProtocolWidgetsModule& RCWidgetsModule = IRemoteControlProtocolWidgetsModule::Get();
+		const TSharedPtr<IRCProtocolBindingList> BindingList = RCWidgetsModule.GetProtocolBindingList();
+		if (!BindingList.IsValid())
+		{
+			return;
+		}
+
+		URemoteControlPreset* Preset = BindingList.IsValid() ? BindingList->GetPreset() : nullptr;
+		const URemoteControlDMXUserData* DMXUserData = URemoteControlDMXUserData::GetOrCreateDMXUserData(Preset);
+		if (!DMXUserData)
+		{
+			return;
+		}
+
+		// Only auto bind when auto patching is disabled
+		if (!DMXUserData->IsAutoPatch())
+		{
+			for (const TSharedPtr<TStructOnScope<FRemoteControlProtocolEntity>>& Entity : BindingList->GetAwaitingProtocolEntities())
+			{
+				ProcessAutoBinding(Entity);
+			}
+		}
+	}
+
+	TStatId FRemoteControlDMXAutoBindHandler::GetStatId() const
+	{
+		RETURN_QUICK_DECLARE_CYCLE_STAT(FRemoteControlDMXAutoBindHandler, STATGROUP_Tickables);
+	}
+
+	void FRemoteControlDMXAutoBindHandler::ProcessAutoBinding(const TSharedPtr<TStructOnScope<FRemoteControlProtocolEntity>>& ProtocolEntity)
+	{
+		// Bind only in editor 
+		if (!GIsEditor)
+		{
+			return;
+		}
+
+		// Bind only when auto-patching is disabled
+		const FRemoteControlDMXProtocolEntity* DMXProtocolEntity = ProtocolEntity.IsValid() ? ProtocolEntity->Cast<FRemoteControlDMXProtocolEntity>() : nullptr;
+		URemoteControlPreset* Preset = DMXProtocolEntity ? DMXProtocolEntity->GetOwner().Get() : nullptr;
+		const URemoteControlDMXUserData* DMXUserData = URemoteControlDMXUserData::GetOrCreateDMXUserData(Preset);
+		if (!DMXProtocolEntity ||
+			!Preset ||
+			!DMXUserData ||
+			DMXUserData->IsAutoPatch())
+		{
+			return;
+		}
+
+
+		UDMXEntityFixturePatch* FixturePatch = DMXProtocolEntity->ExtraSetting.FixturePatchReference.GetFixturePatch();
+		const UDMXLibrary* DMXLibrary = FixturePatch ? FixturePatch->GetParentLibrary() : nullptr;
+		if (!FixturePatch || !DMXLibrary)
+		{
+			return;
+		}
+
+		// Gather DMX Signals per port
+		int32 NumSignals = 0;
+		TMap<TSharedRef<FDMXInputPort>, TArray<TSharedRef<FDMXSignal>>> InputPortToSignalsMap;
+		for (const TSharedRef<FDMXInputPort>& InputPort : DMXLibrary->GetInputPorts())
+		{
+			for (const TTuple<int32, TSharedPtr<FDMXSignal>>& ExternUniverseToSignalPair : InputPort->GameThreadGetAllDMXSignals())
+			{
+				if (ExternUniverseToSignalPair.Value.IsValid())
+				{
+					InputPortToSignalsMap.FindOrAdd(InputPort).Add(ExternUniverseToSignalPair.Value.ToSharedRef());
+					NumSignals++;
+				}
+			}
+		}
+
+		for (const TTuple<TSharedRef<FDMXInputPort>, TArray<TSharedRef<FDMXSignal>>>& InputPortToSignalPair : InputPortToSignalsMap)
+		{
+			// Ignore first received data 
+			const TArray<TSharedRef<FDMXSignal>>* OldSignalsPtr = LastInputPortToSignalsMap.Find(InputPortToSignalPair.Key);
+			if (!OldSignalsPtr)
+			{
+				continue;
+			}
+			const TArray<TSharedRef<FDMXSignal>>& OldSignals = *OldSignalsPtr;
+			const TArray<TSharedRef<FDMXSignal>>& NewSignals = InputPortToSignalPair.Value;
+
+			// Ignore additional and unchanged signals
+			if (OldSignals.Num() != NewSignals.Num() ||
+				OldSignals == NewSignals)
+			{
+				continue;
+			}
+
+			// Find any changed universe and channel
+			bool bFoundAutoBinding = false;
+			for (int32 SignalIndex = 0; SignalIndex < OldSignals.Num(); SignalIndex++)
+			{
+				const TSharedRef<FDMXSignal> OldSignal = OldSignals[SignalIndex];
+				const TSharedRef<FDMXSignal> NewSignal = NewSignals[SignalIndex];
+
+				// Only test if data is of same size
+				if (OldSignal->ChannelData.IsEmpty() ||
+					OldSignal->ChannelData.Num() != NewSignal->ChannelData.Num())
+				{
+					continue;
+				}
+
+				const bool bDataChanged = FMemory::Memcmp(OldSignal->ChannelData.GetData(), NewSignal->ChannelData.GetData(), OldSignal->ChannelData.Num()) != 0;
+				if (!bDataChanged)
+				{
+					continue;
+				}
+
+				// Find the changed channel
+				TOptional<int32> ChangedChannel;
+				for (int32 ChannelIndex = 0; ChannelIndex < OldSignal->ChannelData.Num(); ChannelIndex++)
+				{
+					if (OldSignal->ChannelData[ChannelIndex] != NewSignal->ChannelData[ChannelIndex])
+					{
+						ChangedChannel = ChannelIndex + 1;
+						break;
+					}
+				}
+
+				if (ChangedChannel.IsSet())
+				{
+					const TSharedPtr<IRemoteControlProtocol> Protocol = IRemoteControlProtocolModule::Get().GetProtocolByName(FRemoteControlProtocolDMX::ProtocolName);
+					if (!Protocol.IsValid())
+					{
+						return;
+					}
+
+					const int32 Universe = InputPortToSignalPair.Key->ConvertExternToLocalUniverseID(NewSignal->ExternUniverseID);
+					Protocol->Unbind(ProtocolEntity);
+
+					FixturePatch->SetUniverseID(Universe);
+					FixturePatch->SetStartingChannel(ChangedChannel.GetValue());
+
+					Protocol->Bind(ProtocolEntity);
+
+					bFoundAutoBinding = true;
+					break;
+				}
+			}
+
+			if (bFoundAutoBinding)
+			{
+				break;
+			}
+		}
+
+		// Remember the new signals
+		LastInputPortToSignalsMap = InputPortToSignalsMap;
+	}
+}
diff --git a/Engine/Plugins/VirtualProduction/RemoteControlProtocolDMX/Source/RemoteControlProtocolDMXEditor/Private/RemoteControlDMXAutoBindHandler.h b/Engine/Plugins/VirtualProduction/RemoteControlProtocolDMX/Source/RemoteControlProtocolDMXEditor/Private/RemoteControlDMXAutoBindHandler.h
new file mode 100644
--- /dev/null
+++ b/Engine/Plugins/VirtualProduction/RemoteControlProtocolDMX/Source/RemoteControlProtocolDMXEditor/Private/RemoteControlDMXAutoBindHandler.h
@@ -0,0 +1,44 @@
+// Copyright Epic Games, Inc. All Rights Reserved.
+
+#pragma once
+
+#include "Tickable.h"
+#include "Stats/Stats.h"
+
+class FDMXInputPort;
+class FDMXSignal;
+class URemoteControlPreset;
+struct FRemoteControlProtocolEntity;
+template<typename T> class TStructOnScope;
+
+namespace UE::RemoteControl::DMX
+{
+	/** Handles auto binding DMX values and maintains the related list */
+	class FRemoteControlDMXAutoBindHandler
+		: public FTickableGameObject
+	{
+	public:
+		/** Registers the auto-bind handler with the engine */
+		static void Register();
+
+	private:
+		//~ Begin FTickableGameObject Interface
+		virtual bool IsTickable() const override { return true; }
+		virtual bool IsTickableInEditor() const override { return true; }
+		virtual void Tick(float DeltaTime) override;
+		virtual TStatId GetStatId() const override;
+		//~ End FTickableGameObject Interface
+		
+		/**
+		 * Processes auto binding
+		 * @param InProtocolEntityPtr	Protocol entity pointer
+		 */
+		void ProcessAutoBinding(const TSharedPtr<TStructOnScope<FRemoteControlProtocolEntity>>& ProtocolEntity);
+
+		/** Last received DMX Signals with their input ports */
+		TMap<TSharedRef<FDMXInputPort>, TArray<TSharedRef<FDMXSignal>>> LastInputPortToSignalsMap;
+
+		/** Caches if auto patch mode is enabled */
+		bool bIsAutoPatch = true;
+	};
+}
diff --git a/Engine/Plugins/VirtualProduction/RemoteControlProtocolDMX/Source/RemoteControlProtocolDMXEditor/Private/RemoteControlDMXEditorLog.h b/Engine/Plugins/VirtualProduction/RemoteControlProtocolDMX/Source/RemoteControlProtocolDMXEditor/Private/RemoteControlDMXEditorLog.h
new file mode 100644
--- /dev/null
+++ b/Engine/Plugins/VirtualProduction/RemoteControlProtocolDMX/Source/RemoteControlProtocolDMXEditor/Private/RemoteControlDMXEditorLog.h
@@ -0,0 +1,7 @@
+// Copyright Epic Games, Inc. All Rights Reserved.
+
+#pragma once
+
+#include "Logging/LogMacros.h"
+
+DECLARE_LOG_CATEGORY_CLASS(LogRemoteControlDMXEditor, Log, All);
diff --git a/Engine/Plugins/VirtualProduction/RemoteControlProtocolDMX/Source/RemoteControlProtocolDMXEditor/Private/RemoteControlDMXExposedEntitiesGroupPatchWidgetFactory.cpp b/Engine/Plugins/VirtualProduction/RemoteControlProtocolDMX/Source/RemoteControlProtocolDMXEditor/Private/RemoteControlDMXExposedEntitiesGroupPatchWidgetFactory.cpp
new file mode 100644
--- /dev/null
+++ b/Engine/Plugins/VirtualProduction/RemoteControlProtocolDMX/Source/RemoteControlProtocolDMXEditor/Private/RemoteControlDMXExposedEntitiesGroupPatchWidgetFactory.cpp
@@ -0,0 +1,31 @@
+// Copyright Epic Games, Inc. All Rights Reserved.
+
+#include "RemoteControlDMXExposedEntitiesGroupPatchWidgetFactory.h"
+
+#include "IRemoteControlUIModule.h"
+#include "RemoteControlProtocolDMX.h"
+#include "UI/SRemoteControlDMXExposedEntitiesGroupPatch.h"
+
+namespace UE::RemoteControl::DMX
+{
+	void FRemoteControlDMXExposedEntitiesGroupPatchWidgetFactory::Register()
+	{
+		IRemoteControlUIModule& RemoteControlUIModule = FModuleManager::LoadModuleChecked<IRemoteControlUIModule>("RemoteControlUI");
+		RemoteControlUIModule.RegisterExposedEntitiesGroupWidgetFactory(MakeShared<FRemoteControlDMXExposedEntitiesGroupPatchWidgetFactory>());
+	}
+
+	FName FRemoteControlDMXExposedEntitiesGroupPatchWidgetFactory::GetColumnName() const
+	{
+		return FRemoteControlProtocolDMX::PatchColumnName;
+	}
+
+	FName FRemoteControlDMXExposedEntitiesGroupPatchWidgetFactory::GetProtocolName() const
+	{
+		return FRemoteControlProtocolDMX::ProtocolName;
+	}
+
+	TSharedRef<SWidget> FRemoteControlDMXExposedEntitiesGroupPatchWidgetFactory::MakeWidget(const FRCPanelExposedEntitiesGroupWidgetFactoryArgs& Args)
+	{
+		return SNew(SRemoteControlDMXExposedEntitiesGroupPatch, Args.WeakPreset, Args.ChildProperties);
+	}
+}
diff --git a/Engine/Plugins/VirtualProduction/RemoteControlProtocolDMX/Source/RemoteControlProtocolDMXEditor/Private/RemoteControlDMXExposedEntitiesGroupPatchWidgetFactory.h b/Engine/Plugins/VirtualProduction/RemoteControlProtocolDMX/Source/RemoteControlProtocolDMXEditor/Private/RemoteControlDMXExposedEntitiesGroupPatchWidgetFactory.h
new file mode 100644
--- /dev/null
+++ b/Engine/Plugins/VirtualProduction/RemoteControlProtocolDMX/Source/RemoteControlProtocolDMXEditor/Private/RemoteControlDMXExposedEntitiesGroupPatchWidgetFactory.h
@@ -0,0 +1,31 @@
+// Copyright Epic Games, Inc. All Rights Reserved.
+
+#pragma once
+
+#include "Templates/SharedPointer.h"
+#include "UI/IRCPanelExposedEntitiesGroupWidgetFactory.h"
+
+class SWidget;
+struct FRCFieldGroupColumnExtenderInfo;
+
+namespace UE::RemoteControl::DMX
+{
+	/** Widget factory the patch column in an exposed entities group row */
+	class FRemoteControlDMXExposedEntitiesGroupPatchWidgetFactory
+		: public IRCPanelExposedEntitiesGroupWidgetFactory
+		, public TSharedFromThis<FRemoteControlDMXExposedEntitiesGroupPatchWidgetFactory>
+	{
+	public:
+		virtual ~FRemoteControlDMXExposedEntitiesGroupPatchWidgetFactory() {};
+
+		/** Registers this factory with the engine */
+		static void Register();
+
+	protected:
+		//~ Begin IRCPanelExposedEntitiesGroupWidgetFactory interface
+		virtual FName GetColumnName() const override;
+		virtual FName GetProtocolName() const override;
+		virtual TSharedRef<SWidget> MakeWidget(const FRCPanelExposedEntitiesGroupWidgetFactoryArgs& Args) override;
+		//~ End IRCPanelExposedEntitiesGroupWidgetFactory
+	};
+}
diff --git a/Engine/Plugins/VirtualProduction/RemoteControlProtocolDMX/Source/RemoteControlProtocolDMXEditor/Private/RemoteControlDMXExposedEntitiesListSettings.cpp b/Engine/Plugins/VirtualProduction/RemoteControlProtocolDMX/Source/RemoteControlProtocolDMXEditor/Private/RemoteControlDMXExposedEntitiesListSettings.cpp
new file mode 100644
--- /dev/null
+++ b/Engine/Plugins/VirtualProduction/RemoteControlProtocolDMX/Source/RemoteControlProtocolDMXEditor/Private/RemoteControlDMXExposedEntitiesListSettings.cpp
@@ -0,0 +1,106 @@
+// Copyright Epic Games, Inc. All Rights Reserved.
+
+#include "RemoteControlDMXExposedEntitiesListSettings.h"
+
+#include "Algo/Find.h"
+#include "Framework/Notifications/NotificationManager.h"
+#include "IRemoteControlUIModule.h"
+#include "Library/DMXEntityFixturePatch.h"
+#include "Library/RemoteControlDMXLibraryProxy.h"
+#include "RemoteControlDMXEditorLog.h"
+#include "RemoteControlDMXPerPresetEditorSettings.h"
+#include "RemoteControlDMXUserData.h"
+#include "RemoteControlPreset.h"
+#include "RemoteControlProtocolDMX.h"
+#include "Widgets/Notifications/SNotificationList.h"
+
+#define LOCTEXT_NAMESPACE "RemoteControlDMXExposedEntitiesListSettings"
+
+namespace UE::RemoteControl::DMX
+{
+	void FRemoteControlDMXExposedEntitiesListSettings::Register()
+	{
+		static TSharedRef<IRCPanelExposedEntitiesListSettingsForProtocol> Instance = MakeShared<FRemoteControlDMXExposedEntitiesListSettings>();
+		IRemoteControlUIModule::Get().RegisterExposedEntitiesListSettingsForProtocol(Instance);
+	}
+
+	FName FRemoteControlDMXExposedEntitiesListSettings::GetProtocolName() const
+	{
+		return FRemoteControlProtocolDMX::ProtocolName;
+	}
+
+	FRCPanelExposedEntitiesListSettingsData FRemoteControlDMXExposedEntitiesListSettings::GetListSettings(URemoteControlPreset* Preset) const
+	{
+		const URemoteControlDMXPerPresetEditorSettings* EditorSettings = URemoteControlDMXPerPresetEditorSettings::GetOrCreatePerPresetEditorSettings(Preset);
+		if (EditorSettings)
+		{
+			return EditorSettings->ExposedEntitiesListSettings;
+		}
+
+		return FRCPanelExposedEntitiesListSettingsData();
+	}
+	
+	void FRemoteControlDMXExposedEntitiesListSettings::OnSettingsChanged(URemoteControlPreset* Preset, const FRCPanelExposedEntitiesListSettingsData& ListSettings)
+	{
+		URemoteControlDMXUserData* DMXUserData = GetUserData(Preset);
+		URemoteControlDMXLibraryProxy* DMXLibraryProxy = DMXUserData ? DMXUserData->GetDMXLibraryProxy() : nullptr;
+		if (Preset && DMXUserData && DMXLibraryProxy)
+		{
+			// Store in data
+			URemoteControlDMXPerPresetEditorSettings* EditorSettings = URemoteControlDMXPerPresetEditorSettings::GetOrCreatePerPresetEditorSettings(Preset);
+			if (EditorSettings)
+			{
+				EditorSettings->Modify();
+				EditorSettings->ExposedEntitiesListSettings = ListSettings;
+			}
+
+			// Update patch group mode
+			const ERemoteControlDMXPatchGroupMode PreviousPatchGroupMode = DMXUserData->GetPatchGroupMode();
+
+			const bool bGroupByOwner = DMXUserData->IsAutoPatch() && ListSettings.FieldGroupType == ERCFieldGroupType::Owner;
+			const ERemoteControlDMXPatchGroupMode NewPatchGroupMode = bGroupByOwner ? ERemoteControlDMXPatchGroupMode::GroupByOwner : ERemoteControlDMXPatchGroupMode::GroupByProperty;
+
+			const bool bPatchGroupModeChanged = PreviousPatchGroupMode != NewPatchGroupMode;
+			if (bPatchGroupModeChanged)
+			{
+				DMXUserData->Modify();
+				DMXUserData->SetPatchGroupMode(NewPatchGroupMode);
+			}
+
+			// Show a warning if not all patches can be created in the current list order
+			const TArray<UDMXEntityFixturePatch*> FixturePatchesThatExceedUniverseSize = DMXLibraryProxy->FindPatchesThatExceedUniverseSize();
+			if (!FixturePatchesThatExceedUniverseSize.IsEmpty())
+			{
+				const FText WarningInfo = LOCTEXT("PatchesExceedUniverse", "Remote control generated DMX patches exceed the universe size. See log for details.");
+
+				FNotificationInfo Info(WarningInfo);
+				Info.bUseSuccessFailIcons = true;
+				Info.ExpireDuration = 10.f;
+
+				FSlateNotificationManager::Get().AddNotification(Info);
+
+				for (const UDMXEntityFixturePatch* FixturePatch : FixturePatchesThatExceedUniverseSize)
+				{
+					UE_LOG(LogRemoteControlDMXEditor, Warning, TEXT("Remote control generated patch '%s' exceeds 512 channels. Only the first 512 channels will be available."), *FixturePatch->Name);
+				}
+			}
+		}
+	}
+
+	URemoteControlDMXUserData* FRemoteControlDMXExposedEntitiesListSettings::GetUserData(URemoteControlPreset* Preset) const
+	{
+		if (!Preset)
+		{
+			return nullptr;
+		}
+			
+		const TObjectPtr<UObject>* DMXUserDataObjectPtr = Algo::FindByPredicate(Preset->UserData, [](const TObjectPtr<UObject>& Object)
+			{
+				return Object && Object->GetClass() == URemoteControlDMXUserData::StaticClass();
+			});
+
+		return DMXUserDataObjectPtr ? CastChecked<URemoteControlDMXUserData>(*DMXUserDataObjectPtr) : nullptr;
+	}
+}
+
+#undef LOCTEXT_NAMESPACE
diff --git a/Engine/Plugins/VirtualProduction/RemoteControlProtocolDMX/Source/RemoteControlProtocolDMXEditor/Private/RemoteControlDMXExposedEntitiesListSettings.h b/Engine/Plugins/VirtualProduction/RemoteControlProtocolDMX/Source/RemoteControlProtocolDMXEditor/Private/RemoteControlDMXExposedEntitiesListSettings.h
new file mode 100644
--- /dev/null
+++ b/Engine/Plugins/VirtualProduction/RemoteControlProtocolDMX/Source/RemoteControlProtocolDMXEditor/Private/RemoteControlDMXExposedEntitiesListSettings.h
@@ -0,0 +1,33 @@
+// Copyright Epic Games, Inc. All Rights Reserved.
+
+#pragma once
+
+#include "Templates/SharedPointer.h"
+#include "UI/IRCPanelExposedEntitiesListSettingsForProtocol.h"
+
+class URemoteControlDMXUserData;
+
+namespace UE::RemoteControl::DMX
+{
+	/** Settings for the DMX exposed entities list */
+	class FRemoteControlDMXExposedEntitiesListSettings
+		: public IRCPanelExposedEntitiesListSettingsForProtocol
+		, public TSharedFromThis<FRemoteControlDMXExposedEntitiesListSettings>
+	{
+	public:
+		virtual ~FRemoteControlDMXExposedEntitiesListSettings() = default;
+
+		/** Registers the exposed entities list settings with the engine */
+		static void Register();
+
+		//~ Begin IRCPanelExposedEntityListSettings
+		virtual FName GetProtocolName() const override;
+		virtual FRCPanelExposedEntitiesListSettingsData GetListSettings(URemoteControlPreset* Preset) const override;
+		virtual void OnSettingsChanged(URemoteControlPreset* Preset, const FRCPanelExposedEntitiesListSettingsData& ListSettings) override;
+		//~ End IRCPanelExposedEntityListSettings
+
+	private:
+		/** Returns the DMX user data of the preset */
+		URemoteControlDMXUserData* GetUserData(URemoteControlPreset* Preset) const;
+	};
+}
diff --git a/Engine/Plugins/VirtualProduction/RemoteControlProtocolDMX/Source/RemoteControlProtocolDMXEditor/Private/RemoteControlDMXExposedEntitiesPanelExtender.cpp b/Engine/Plugins/VirtualProduction/RemoteControlProtocolDMX/Source/RemoteControlProtocolDMXEditor/Private/RemoteControlDMXExposedEntitiesPanelExtender.cpp
new file mode 100644
--- /dev/null
+++ b/Engine/Plugins/VirtualProduction/RemoteControlProtocolDMX/Source/RemoteControlProtocolDMXEditor/Private/RemoteControlDMXExposedEntitiesPanelExtender.cpp
@@ -0,0 +1,69 @@
+// Copyright Epic Games, Inc. All Rights Reserved.
+
+#include "RemoteControlDMXExposedEntitiesPanelExtender.h"
+
+#include "Algo/Find.h"
+#include "IRemoteControlUIModule.h"
+#include "Modules/ModuleManager.h"
+#include "RemoteControlDMXUserData.h"
+#include "RemoteControlPreset.h"
+#include "RemoteControlProtocolDMX.h"
+#include "Styling/AppStyle.h"
+#include "UI/SRemoteControlDMXPresetUserData.h"
+
+namespace UE::RemoteControl::DMX
+{
+	void FRemoteControlDMXExposedEntitiesPanelExtender::Register()
+	{
+		IRemoteControlUIModule& RemoteControlUIModule = FModuleManager::LoadModuleChecked<IRemoteControlUIModule>("RemoteControlUI");
+		RemoteControlUIModule.RegisterExposedEntitiesPanelExtender(MakeShared<FRemoteControlDMXExposedEntitiesPanelExtender>());
+	}
+
+	TSharedRef<SWidget> FRemoteControlDMXExposedEntitiesPanelExtender::MakeWidget(URemoteControlPreset* Preset, const IRCExposedEntitiesPanelExtender::FArgs& Args)
+	{
+		if (!ensureMsgf(Preset, TEXT("Unexpected invalid preset provided when trying to extend the Exposed Entities Panel, cannot create exposed entities panel for DMX")) ||
+			!ensureMsgf(Args.ActiveProtocolAttribute.IsBound(), TEXT("Unexpected ActiveProtocolAttribute is not bound, cannot create exposed entities panel for DMX.")))
+		{
+			return SNullWidget::NullWidget;
+		}
+
+		// Get or create DMX user data for this preset
+		URemoteControlDMXUserData* DMXUserData = [Preset]()
+			{
+				TObjectPtr<UObject>* DMXUserDataObjectPtr = Algo::FindByPredicate(Preset->UserData, [](const TObjectPtr<UObject>& Object)
+					{
+						return Object && Object->GetClass() == URemoteControlDMXUserData::StaticClass();
+					});
+
+				if (DMXUserDataObjectPtr)
+				{
+					return CastChecked<URemoteControlDMXUserData>(*DMXUserDataObjectPtr);
+				}
+				else
+				{
+					URemoteControlDMXUserData* NewDMXUserData = NewObject<URemoteControlDMXUserData>(Preset);
+
+					Preset->Modify();
+					Preset->UserData.Add(NewDMXUserData);
+
+					return NewDMXUserData;
+				}
+			}();
+			check(DMXUserData);
+
+		// Create the extension widget
+		const TSharedRef<SWidget> ExtensionWidget =
+			SNew(SBorder)
+			.BorderImage(FAppStyle::GetBrush("NoBorder"))
+			.Visibility_Lambda([Args]()
+				{
+					const bool bIsActiveProtocolDMX = Args.ActiveProtocolAttribute.IsBound() && Args.ActiveProtocolAttribute.Get() == FRemoteControlProtocolDMX::ProtocolName;
+					return bIsActiveProtocolDMX ? EVisibility::Visible : EVisibility::Collapsed;
+				})
+			[
+				SNew(SRemoteControlDMXPresetUserData, DMXUserData)
+			];
+
+		return ExtensionWidget;
+	}
+}
diff --git a/Engine/Plugins/VirtualProduction/RemoteControlProtocolDMX/Source/RemoteControlProtocolDMXEditor/Private/RemoteControlDMXExposedEntitiesPanelExtender.h b/Engine/Plugins/VirtualProduction/RemoteControlProtocolDMX/Source/RemoteControlProtocolDMXEditor/Private/RemoteControlDMXExposedEntitiesPanelExtender.h
new file mode 100644
--- /dev/null
+++ b/Engine/Plugins/VirtualProduction/RemoteControlProtocolDMX/Source/RemoteControlProtocolDMXEditor/Private/RemoteControlDMXExposedEntitiesPanelExtender.h
@@ -0,0 +1,28 @@
+// Copyright Epic Games, Inc. All Rights Reserved.
+
+#pragma once
+
+#include "Templates/SharedPointer.h"
+#include "UI/IRCExposedEntitiesPanelExtender.h"
+
+class URemoteControlPreset;
+
+namespace UE::RemoteControl::DMX
+{
+	/** Handles extending the exposed entities panel for DMX */
+	class FRemoteControlDMXExposedEntitiesPanelExtender
+		: public IRCExposedEntitiesPanelExtender
+		, public TSharedFromThis<FRemoteControlDMXExposedEntitiesPanelExtender>
+	{
+	public:
+		virtual ~FRemoteControlDMXExposedEntitiesPanelExtender() {};
+
+		/** Registers the extender with the engine */
+		static void Register();
+
+	protected:
+		//~ Begin IRCExposedEntitiesPanelExtender interface
+		virtual TSharedRef<SWidget> MakeWidget(URemoteControlPreset* Preset, const IRCExposedEntitiesPanelExtender::FArgs& Args) override;
+		//~ End IRCExposedEntitiesPanelExtender
+	};
+}
diff --git a/Engine/Plugins/VirtualProduction/RemoteControlProtocolDMX/Source/RemoteControlProtocolDMXEditor/Private/RemoteControlDMXExposedEntityPatchWidgetFactory.cpp b/Engine/Plugins/VirtualProduction/RemoteControlProtocolDMX/Source/RemoteControlProtocolDMXEditor/Private/RemoteControlDMXExposedEntityPatchWidgetFactory.cpp
new file mode 100644
--- /dev/null
+++ b/Engine/Plugins/VirtualProduction/RemoteControlProtocolDMX/Source/RemoteControlProtocolDMXEditor/Private/RemoteControlDMXExposedEntityPatchWidgetFactory.cpp
@@ -0,0 +1,31 @@
+// Copyright Epic Games, Inc. All Rights Reserved.
+
+#include "RemoteControlDMXExposedEntityPatchWidgetFactory.h"
+
+#include "IRemoteControlUIModule.h"
+#include "RemoteControlProtocolDMX.h"
+#include "UI/SRemoteControlDMXExposedEntityPatch.h"
+
+namespace UE::RemoteControl::DMX
+{
+	void FRemoteControlDMXExposedEntityPatchWidgetFactory::Register()
+	{
+		IRemoteControlUIModule& RemoteControlUIModule = FModuleManager::LoadModuleChecked<IRemoteControlUIModule>("RemoteControlUI");
+		RemoteControlUIModule.RegisterExposedEntityWidgetFactory(MakeShared<FRemoteControlDMXExposedEntityPatchWidgetFactory>());
+	}
+
+	FName FRemoteControlDMXExposedEntityPatchWidgetFactory::GetColumnName() const
+	{
+		return FRemoteControlProtocolDMX::PatchColumnName;
+	}
+
+	FName FRemoteControlDMXExposedEntityPatchWidgetFactory::GetProtocolName() const
+	{
+		return FRemoteControlProtocolDMX::ProtocolName;
+	}
+
+	TSharedRef<SWidget> FRemoteControlDMXExposedEntityPatchWidgetFactory::MakePropertyWidget(const FRCPanelExposedPropertyWidgetArgs& Args)
+	{
+		return SNew(SRemoteControlDMXExposedEntityPatch, Args.WeakPreset, Args.Property);
+	}
+}
diff --git a/Engine/Plugins/VirtualProduction/RemoteControlProtocolDMX/Source/RemoteControlProtocolDMXEditor/Private/RemoteControlDMXExposedEntityPatchWidgetFactory.h b/Engine/Plugins/VirtualProduction/RemoteControlProtocolDMX/Source/RemoteControlProtocolDMXEditor/Private/RemoteControlDMXExposedEntityPatchWidgetFactory.h
new file mode 100644
--- /dev/null
+++ b/Engine/Plugins/VirtualProduction/RemoteControlProtocolDMX/Source/RemoteControlProtocolDMXEditor/Private/RemoteControlDMXExposedEntityPatchWidgetFactory.h
@@ -0,0 +1,27 @@
+// Copyright Epic Games, Inc. All Rights Reserved.
+
+#pragma once
+
+#include "UI/IRCPanelExposedEntityWidgetFactory.h"
+
+namespace UE::RemoteControl::DMX
+{
+	/** Widget factory the patch column in an exposed entities group row */
+	class FRemoteControlDMXExposedEntityPatchWidgetFactory
+		: public IRCPanelExposedEntityWidgetFactory
+		, public TSharedFromThis<FRemoteControlDMXExposedEntityPatchWidgetFactory>
+	{
+	public:
+		virtual ~FRemoteControlDMXExposedEntityPatchWidgetFactory() {};
+
+		/** Registers the factory with the engine */
+		static void Register();
+
+	protected:
+		//~ Begin IRCPanelExposedEntitiesGroupWidgetFactory interface
+		virtual FName GetColumnName() const override;
+		virtual FName GetProtocolName() const override;
+		virtual TSharedRef<SWidget> MakePropertyWidget(const FRCPanelExposedPropertyWidgetArgs& Args) override;
+		//~ End IRCPanelExposedEntitiesGroupWidgetFactory
+	};
+}
diff --git a/Engine/Plugins/VirtualProduction/RemoteControlProtocolDMX/Source/RemoteControlProtocolDMXEditor/Private/RemoteControlDMXPerPresetEditorSettings.cpp b/Engine/Plugins/VirtualProduction/RemoteControlProtocolDMX/Source/RemoteControlProtocolDMXEditor/Private/RemoteControlDMXPerPresetEditorSettings.cpp
new file mode 100644
--- /dev/null
+++ b/Engine/Plugins/VirtualProduction/RemoteControlProtocolDMX/Source/RemoteControlProtocolDMXEditor/Private/RemoteControlDMXPerPresetEditorSettings.cpp
@@ -0,0 +1,20 @@
+// Copyright Epic Games, Inc. All Rights Reserved.
+
+#include "RemoteControlDMXPerPresetEditorSettings.h"
+
+#include "RemoteControlDMXUserData.h"
+
+URemoteControlDMXPerPresetEditorSettings* URemoteControlDMXPerPresetEditorSettings::GetOrCreatePerPresetEditorSettings(URemoteControlPreset* Preset)
+{
+	URemoteControlDMXUserData* DMXUserData = URemoteControlDMXUserData::GetOrCreateDMXUserData(Preset);
+
+	if (DMXUserData->PerPresetEditorSettings && DMXUserData->PerPresetEditorSettings->GetClass() == URemoteControlDMXPerPresetEditorSettings::StaticClass())
+	{
+		return CastChecked<URemoteControlDMXPerPresetEditorSettings>(DMXUserData->PerPresetEditorSettings);
+	}
+	else
+	{
+		DMXUserData->PerPresetEditorSettings = NewObject<URemoteControlDMXPerPresetEditorSettings>(DMXUserData, NAME_None);
+		return CastChecked<URemoteControlDMXPerPresetEditorSettings>(DMXUserData->PerPresetEditorSettings);
+	}
+}
diff --git a/Engine/Plugins/VirtualProduction/RemoteControlProtocolDMX/Source/RemoteControlProtocolDMXEditor/Private/RemoteControlDMXPerPresetEditorSettings.h b/Engine/Plugins/VirtualProduction/RemoteControlProtocolDMX/Source/RemoteControlProtocolDMXEditor/Private/RemoteControlDMXPerPresetEditorSettings.h
new file mode 100644
--- /dev/null
+++ b/Engine/Plugins/VirtualProduction/RemoteControlProtocolDMX/Source/RemoteControlProtocolDMXEditor/Private/RemoteControlDMXPerPresetEditorSettings.h
@@ -0,0 +1,25 @@
+// Copyright Epic Games, Inc. All Rights Reserved.
+
+#pragma once
+
+#include "UI/RCPanelExposedEntitiesListSettingsData.h"
+#include "UObject/Object.h"
+
+#include "RemoteControlDMXPerPresetEditorSettings.generated.h"
+
+class URemoteControlPreset;
+
+/** UObject to hold entities list settings data */
+UCLASS()
+class URemoteControlDMXPerPresetEditorSettings
+	: public UObject
+{
+	GENERATED_BODY()
+
+public:
+	/** Returns the remote control DMX editor settings for the specified preset */
+	static URemoteControlDMXPerPresetEditorSettings* GetOrCreatePerPresetEditorSettings(URemoteControlPreset* Preset);
+
+	UPROPERTY()
+	FRCPanelExposedEntitiesListSettingsData ExposedEntitiesListSettings;
+};
diff --git a/Engine/Plugins/VirtualProduction/RemoteControlProtocolDMX/Source/RemoteControlProtocolDMXEditor/Private/RemoteControlDMXProtocolEntityExtraSettingCustomization.cpp b/Engine/Plugins/VirtualProduction/RemoteControlProtocolDMX/Source/RemoteControlProtocolDMXEditor/Private/RemoteControlDMXProtocolEntityExtraSettingCustomization.cpp
deleted file mode 100644
diff --git a/Engine/Plugins/VirtualProduction/RemoteControlProtocolDMX/Source/RemoteControlProtocolDMXEditor/Private/RemoteControlDMXProtocolEntityExtraSettingCustomization.h b/Engine/Plugins/VirtualProduction/RemoteControlProtocolDMX/Source/RemoteControlProtocolDMXEditor/Private/RemoteControlDMXProtocolEntityExtraSettingCustomization.h
deleted file mode 100644
diff --git a/Engine/Plugins/VirtualProduction/RemoteControlProtocolDMX/Source/RemoteControlProtocolDMXEditor/Private/RemoteControlProtocolDMXEditorModule.cpp b/Engine/Plugins/VirtualProduction/RemoteControlProtocolDMX/Source/RemoteControlProtocolDMXEditor/Private/RemoteControlProtocolDMXEditorModule.cpp
--- a/Engine/Plugins/VirtualProduction/RemoteControlProtocolDMX/Source/RemoteControlProtocolDMXEditor/Private/RemoteControlProtocolDMXEditorModule.cpp
+++ b/Engine/Plugins/VirtualProduction/RemoteControlProtocolDMX/Source/RemoteControlProtocolDMXEditor/Private/RemoteControlProtocolDMXEditorModule.cpp
@@ -1,18 +1,22 @@
 // Copyright Epic Games, Inc. All Rights Reserved.
 
-
+#include "Library/RemoteControlDMXLibraryBuilder.h"
 #include "Modules/ModuleManager.h"
-
+#include "RemoteControlDMXAutoBindHandler.h"
+#include "RemoteControlDMXExposedEntitiesGroupPatchWidgetFactory.h"
+#include "RemoteControlDMXExposedEntitiesListSettings.h"
+#include "RemoteControlDMXExposedEntitiesPanelExtender.h"
+#include "RemoteControlDMXExposedEntityPatchWidgetFactory.h"
 #include "RemoteControlProtocolDMX.h"
 #include "RemoteControlProtocolDMXSettings.h"
-#include "RemoteControlProtocolDMXSettingsDetails.h"
-#include "RemoteControlDMXProtocolEntityExtraSettingCustomization.h"
 
 /**
  * Remote control protocol DMX editor that allows have editor functionality for the protocol
  */
 class FRemoteControlProtocolDMXEditorModule : public IModuleInterface
 {
+	using FRemoteControlDMXExposedEntitiesPanelExtender = UE::RemoteControl::DMX::FRemoteControlDMXExposedEntitiesPanelExtender;
+
 public:
 	//~ Begin IModuleInterface interface
 	virtual void StartupModule() override;
@@ -22,25 +26,18 @@ public:
 
 void FRemoteControlProtocolDMXEditorModule::StartupModule()
 {
-	FPropertyEditorModule& PropertyEditorModule = FModuleManager::GetModuleChecked<FPropertyEditorModule>("PropertyEditor");
-
-	PropertyEditorModule.RegisterCustomPropertyTypeLayout(
-		FRemoteControlDMXProtocolEntityExtraSetting::StaticStruct()->GetFName(),
-		FOnGetPropertyTypeCustomizationInstance::CreateStatic(FRemoteControlDMXProtocolEntityExtraSettingCustomization::MakeInstance));
-
-	PropertyEditorModule.RegisterCustomClassLayout(
-		URemoteControlProtocolDMXSettings::StaticClass()->GetFName(),
-		FOnGetDetailCustomizationInstance::CreateStatic(FRemoteControlProtocolDMXSettingsDetails::MakeInstance));
+	using namespace UE::RemoteControl::DMX;
+
+	FRemoteControlDMXLibraryBuilder::Register();
+	FRemoteControlDMXAutoBindHandler::Register();
+	FRemoteControlDMXExposedEntitiesListSettings::Register();
+	FRemoteControlDMXExposedEntitiesPanelExtender::Register();
+	FRemoteControlDMXExposedEntitiesGroupPatchWidgetFactory::Register();
+	FRemoteControlDMXExposedEntityPatchWidgetFactory::Register();
 }
 
 void FRemoteControlProtocolDMXEditorModule::ShutdownModule()
 {
-	if (UObjectInitialized() && FModuleManager::Get().IsModuleLoaded(TEXT("PropertyEditor")))
-	{
-		FPropertyEditorModule& PropertyEditorModule = FModuleManager::GetModuleChecked<FPropertyEditorModule>("PropertyEditor");
-		PropertyEditorModule.UnregisterCustomPropertyTypeLayout(FRemoteControlDMXProtocolEntityExtraSetting::StaticStruct()->GetFName());
-		PropertyEditorModule.UnregisterCustomClassLayout(URemoteControlProtocolDMXSettings::StaticClass()->GetFName());
-	}
 }
 
 
diff --git a/Engine/Plugins/VirtualProduction/RemoteControlProtocolDMX/Source/RemoteControlProtocolDMXEditor/Private/RemoteControlProtocolDMXSettingsDetails.cpp b/Engine/Plugins/VirtualProduction/RemoteControlProtocolDMX/Source/RemoteControlProtocolDMXEditor/Private/RemoteControlProtocolDMXSettingsDetails.cpp
deleted file mode 100644
diff --git a/Engine/Plugins/VirtualProduction/RemoteControlProtocolDMX/Source/RemoteControlProtocolDMXEditor/Private/RemoteControlProtocolDMXSettingsDetails.h b/Engine/Plugins/VirtualProduction/RemoteControlProtocolDMX/Source/RemoteControlProtocolDMXEditor/Private/RemoteControlProtocolDMXSettingsDetails.h
deleted file mode 100644
diff --git a/Engine/Plugins/VirtualProduction/RemoteControlProtocolDMX/Source/RemoteControlProtocolDMXEditor/Private/UI/SRemoteControlDMXExposedEntitiesGroupPatch.cpp b/Engine/Plugins/VirtualProduction/RemoteControlProtocolDMX/Source/RemoteControlProtocolDMXEditor/Private/UI/SRemoteControlDMXExposedEntitiesGroupPatch.cpp
new file mode 100644
--- /dev/null
+++ b/Engine/Plugins/VirtualProduction/RemoteControlProtocolDMX/Source/RemoteControlProtocolDMXEditor/Private/UI/SRemoteControlDMXExposedEntitiesGroupPatch.cpp
@@ -0,0 +1,363 @@
+// Copyright Epic Games, Inc. All Rights Reserved.
+
+#include "SRemoteControlDMXExposedEntitiesGroupPatch.h"
+
+#include "Editor.h"
+#include "Framework/MultiBox/MultiBoxBuilder.h"
+#include "Library/DMXEntityFixturePatch.h"
+#include "Library/DMXEntityFixtureType.h"
+#include "Library/RemoteControlDMXLibraryProxy.h"
+#include "RemoteControlDMXUserData.h"
+#include "RemoteControlField.h"
+#include "RemoteControlPreset.h"
+#include "RemoteControlProtocolDMX.h"
+#include "UObject/StructOnScope.h"
+#include "Widgets/Input/SComboButton.h"
+
+#define LOCTEXT_NAMESPACE "SRemoteControlDMXExposedEntitiesGroupPatch"
+
+namespace UE::RemoteControl::DMX
+{
+	void SRemoteControlDMXExposedEntitiesGroupPatch::Construct(
+		const FArguments& InArgs, 
+		const TWeakObjectPtr<URemoteControlPreset>& InWeakPreset,
+		const TArray<TSharedRef<FRemoteControlProperty>>& InChildProperties)
+	{
+		WeakPreset = InWeakPreset;
+		if (!InWeakPreset.IsValid() || InChildProperties.IsEmpty())
+		{
+			return;
+		}
+		URemoteControlPreset* Preset = WeakPreset.Get();
+		check(Preset);
+
+		// Get entities in this group
+		for (const TSharedRef<FRemoteControlProperty>& Property : InChildProperties)
+		{
+			for (const FRemoteControlProtocolBinding& Binding : Property->ProtocolBindings)
+			{
+				const TSharedPtr<TStructOnScope<FRemoteControlProtocolEntity>> Entity = Binding.GetRemoteControlProtocolEntityPtr();
+				if (!Entity.IsValid() || !Entity->IsValid())
+				{
+					continue;
+				}
+
+				FRemoteControlDMXProtocolEntity* DMXEntity = Entity->Cast<FRemoteControlDMXProtocolEntity>();
+				if (!DMXEntity)
+				{
+					continue;
+				}
+
+				Entities.Add(Entity.ToSharedRef());
+			}
+		}
+
+		RequestRefresh();
+
+		URemoteControlDMXLibraryProxy::GetOnPostPropertyPatchesChanged().AddSP(this, &SRemoteControlDMXExposedEntitiesGroupPatch::RequestRefresh);
+		UDMXEntityFixtureType::GetOnFixtureTypeChanged().AddSP(this, &SRemoteControlDMXExposedEntitiesGroupPatch::OnFixtureTypeChanged);
+		UDMXEntityFixturePatch::GetOnFixturePatchChanged().AddSP(this, &SRemoteControlDMXExposedEntitiesGroupPatch::OnFixturePatchChanged);
+		Preset->OnEntityExposed().AddSP(this, &SRemoteControlDMXExposedEntitiesGroupPatch::OnEntityExposedOrUnexposed);
+		Preset->OnEntityUnexposed().AddSP(this, &SRemoteControlDMXExposedEntitiesGroupPatch::OnEntityExposedOrUnexposed);
+		Preset->OnEntityRebind().AddSP(this, &SRemoteControlDMXExposedEntitiesGroupPatch::OnEntityRebind);
+		Preset->OnEntitiesUpdated().AddSP(this, &SRemoteControlDMXExposedEntitiesGroupPatch::OnEntitiesUpdated);
+	}
+
+	void SRemoteControlDMXExposedEntitiesGroupPatch::Refresh()
+	{
+		RefreshTimerHandle.Invalidate();
+
+		// Update the fixture patches that can be used as primary for this patch
+		WeakFixturePatches.Reset();
+
+		const UDMXEntityFixturePatch* FixturePatch = GetFixturePatch();
+		const UDMXEntityFixtureType* FixtureType = FixturePatch ? FixturePatch->GetFixtureType() : nullptr;
+		if (FixturePatch && FixtureType)
+		{
+			const TArray<TSharedRef<TStructOnScope<FRemoteControlProtocolEntity>>> AllEntities = FRemoteControlDMXProtocolEntity::GetAllDMXProtocolEntitiesInPreset(WeakPreset.Get());
+			for (const TSharedRef<TStructOnScope<FRemoteControlProtocolEntity>>& Entity : AllEntities)
+			{
+				const FRemoteControlDMXProtocolEntity* DMXEntity = Entity->IsValid() ? Entity->Cast<FRemoteControlDMXProtocolEntity>() : nullptr;
+				if (!DMXEntity ||
+					!DMXEntity->ExtraSetting.bIsPrimaryPatch)
+				{
+					continue;
+				}
+
+				UDMXEntityFixturePatch* OtherFixturePatch = DMXEntity->ExtraSetting.FixturePatchReference.GetFixturePatch();
+				if (OtherFixturePatch &&
+					OtherFixturePatch != FixturePatch &&
+					OtherFixturePatch->GetFixtureType() == FixtureType)
+				{
+					WeakFixturePatches.AddUnique(OtherFixturePatch);
+				}
+			}
+		}
+
+
+		if (Entities.IsEmpty())
+		{
+			ChildSlot
+			[
+				SNew(STextBlock)
+				.Text(LOCTEXT("GroupCannotBePatchedInfo", "Not bound to DMX"))
+			];
+		}
+		else
+		{
+			ChildSlot
+			[
+				SNew(SComboButton)
+				.ButtonContent()
+				[
+					SNew(STextBlock)
+					.Text_Lambda([this]()
+						{
+							if (const UDMXEntityFixturePatch* FixturePatch = GetFixturePatch())
+							{
+								const FText FixturePatchNameText = FText::FromString(FixturePatch->Name);
+								if (IsPrimaryFixturePatch())
+								{
+									return FixturePatchNameText;
+								}
+								else
+								{
+									return FText::Format(LOCTEXT("SecondaryFixturePatchText", "Same as {0}"), FixturePatchNameText);
+								}
+							}
+							else
+							{
+								return LOCTEXT("NotPatchedInfo", "Not Patched");
+							}
+						})
+				]
+				.MenuContent()
+				[
+					CreateMenu()
+				]
+			];
+		}
+	}
+
+	void SRemoteControlDMXExposedEntitiesGroupPatch::RequestRefresh()
+	{
+		if (!RefreshTimerHandle.IsValid())
+		{
+			RefreshTimerHandle = GEditor->GetTimerManager()->SetTimerForNextTick(FTimerDelegate::CreateSP(this, &SRemoteControlDMXExposedEntitiesGroupPatch::Refresh));
+		}
+	}
+
+	TSharedRef<SWidget> SRemoteControlDMXExposedEntitiesGroupPatch::CreateMenu()
+	{
+		if (!ensureMsgf(!Entities.IsEmpty(), TEXT("Expected valid entities to create fixture patch combo box but got none")))
+		{
+			return SNullWidget::NullWidget;
+		}
+
+		constexpr bool bShouldCloseWindowAfterSelection = true;
+		const TSharedPtr<FUICommandList> NullCommandList = nullptr;
+		FMenuBuilder MenuBuilder(bShouldCloseWindowAfterSelection, NullCommandList);
+		
+		MenuBuilder.SetSearchable(true);
+		MenuBuilder.AddSearchWidget();
+
+		const FName NoExtensionHook = NAME_None;
+		const FText NoLabel = FText::GetEmpty();
+		const FText NoTooltip = FText::GetEmpty();
+
+		// Actions Section
+		MenuBuilder.BeginSection("Actions", LOCTEXT("ActionsSectionLabel", "Actions"));
+		{
+			if (GetFixturePatch())
+			{
+				// Clear Patch
+				MenuBuilder.AddMenuEntry(
+					LOCTEXT("ClearPatchLabel", "Clear"),
+					LOCTEXT("ClearPatchTooltip", "Clears the fixture patch. The properties will not receive DMX when the patch is cleared."),
+					FSlateIcon(),
+					FUIAction
+					(
+						FExecuteAction::CreateSP(this, &SRemoteControlDMXExposedEntitiesGroupPatch::ClearPatch)
+					)
+				);
+			}
+			
+			if(!GetFixturePatch() || !IsPrimaryFixturePatch())
+			{
+				// Generate Patch
+				MenuBuilder.AddMenuEntry(
+					LOCTEXT("GeneratePatchLabel", "Generate Patch"),
+					LOCTEXT("GeneratePatchTooltip", "Generates a new fixture patch for this owner."),
+					FSlateIcon(),
+					FUIAction
+					(
+						FExecuteAction::CreateSP(this, &SRemoteControlDMXExposedEntitiesGroupPatch::GeneratePatch)
+					)
+				);
+			}
+		}
+		MenuBuilder.EndSection();
+
+		// Fixture Patch section
+		if (!WeakFixturePatches.IsEmpty())
+		{
+			MenuBuilder.BeginSection("FixturePatch", LOCTEXT("FixturePatchSectionLabel", "Patch same as.."));
+			{
+				for (const TWeakObjectPtr<UDMXEntityFixturePatch>& WeakFixturePatch : WeakFixturePatches)
+				{
+					UDMXEntityFixturePatch* FixturePatch = WeakFixturePatch.Get();
+					if (!FixturePatch)
+					{
+						continue;
+					}
+
+					MenuBuilder.AddMenuEntry
+					(
+						FText::FromString(FixturePatch->Name),
+						NoTooltip,
+						FSlateIcon(),
+						FUIAction
+						(
+							FExecuteAction::CreateSP(this, &SRemoteControlDMXExposedEntitiesGroupPatch::OnFixturePatchSelected, WeakFixturePatch)
+						),
+						NoExtensionHook,
+						EUserInterfaceActionType::Button
+					);
+				}
+			}
+
+			MenuBuilder.EndSection();
+		}
+
+		return MenuBuilder.MakeWidget();
+	}
+
+	void SRemoteControlDMXExposedEntitiesGroupPatch::ClearPatch()
+	{
+		for (const TSharedRef<TStructOnScope<FRemoteControlProtocolEntity>>& Entity : Entities)
+		{
+			if (FRemoteControlDMXProtocolEntity* DMXEntity = Entity->Cast<FRemoteControlDMXProtocolEntity>())
+			{
+				DMXEntity->ExtraSetting.bIsPrimaryPatch = true;
+				DMXEntity->ExtraSetting.FixturePatchReference = nullptr;
+				DMXEntity->ExtraSetting.bRequestClearPatch = true;
+
+				DMXEntity->Invalidate();
+			}
+		}
+	}
+
+	void SRemoteControlDMXExposedEntitiesGroupPatch::GeneratePatch()
+	{
+		for (const TSharedRef<TStructOnScope<FRemoteControlProtocolEntity>>& Entity : Entities)
+		{
+			if (FRemoteControlDMXProtocolEntity* DMXEntity = Entity->Cast<FRemoteControlDMXProtocolEntity>())
+			{
+				DMXEntity->ExtraSetting.bIsPrimaryPatch = true;
+				DMXEntity->ExtraSetting.bRequestClearPatch = false;
+
+				DMXEntity->Invalidate();
+			}
+		}
+	}
+
+	void SRemoteControlDMXExposedEntitiesGroupPatch::SetAutoPatchEnabled(bool bEnabled)
+	{
+		if (URemoteControlDMXUserData* DMXUserData = GetDMXUserData())
+		{
+			DMXUserData->SetAutoPatchEnabled(bEnabled);
+		}
+	}
+
+	void SRemoteControlDMXExposedEntitiesGroupPatch::OnFixturePatchSelected(TWeakObjectPtr<UDMXEntityFixturePatch> WeakFixturePatch)
+	{
+		if (!ensureMsgf(WeakFixturePatch.IsValid(), TEXT("Trying to set fixture patch for remote control owner, but fixture patch is no longer valid.")))
+		{
+			return;
+		}
+
+		for (const TSharedRef<TStructOnScope<FRemoteControlProtocolEntity>>& Entity : Entities)
+		{
+			if (FRemoteControlDMXProtocolEntity* DMXEntity = Entity->Cast<FRemoteControlDMXProtocolEntity>())
+			{
+				if (DMXEntity->ExtraSetting.FixturePatchReference.GetFixturePatch() != WeakFixturePatch.Get())
+				{
+					DMXEntity->ExtraSetting.FixturePatchReference = WeakFixturePatch.Get();
+					DMXEntity->ExtraSetting.bIsPrimaryPatch = false;
+					DMXEntity->ExtraSetting.bRequestClearPatch = false;
+
+					DMXEntity->Invalidate();
+				}
+			}
+		}
+	}
+
+	bool SRemoteControlDMXExposedEntitiesGroupPatch::IsPrimaryFixturePatch() const
+	{
+		if (Entities.IsEmpty() || !Entities[0]->IsValid())
+		{
+			return false;
+		}
+
+		const FRemoteControlDMXProtocolEntity* DMXEntity = Entities[0]->Cast<FRemoteControlDMXProtocolEntity>();
+		return DMXEntity ? DMXEntity->ExtraSetting.bIsPrimaryPatch : false;
+	}
+
+	UDMXEntityFixturePatch* SRemoteControlDMXExposedEntitiesGroupPatch::GetFixturePatch() const
+	{			
+		if (Entities.IsEmpty() || !Entities[0]->IsValid())
+		{
+			return nullptr;
+		}
+
+		const FRemoteControlDMXProtocolEntity* DMXEntity = Entities[0]->Cast<FRemoteControlDMXProtocolEntity>();
+		return DMXEntity ? DMXEntity->ExtraSetting.FixturePatchReference.GetFixturePatch() : nullptr;
+	}
+
+	URemoteControlDMXUserData* SRemoteControlDMXExposedEntitiesGroupPatch::GetDMXUserData() const
+	{
+		if (WeakPreset.IsValid())
+		{
+			return nullptr;
+		}
+
+		const TObjectPtr<UObject>* DMXUserDataObjectPtr = Algo::FindByPredicate(WeakPreset.Get()->UserData, [](const TObjectPtr<UObject>& Object)
+			{
+				return Object && Object->GetClass() == URemoteControlDMXUserData::StaticClass();
+			});
+
+		return DMXUserDataObjectPtr ? CastChecked<URemoteControlDMXUserData>(*DMXUserDataObjectPtr) : nullptr;
+	}
+
+	void SRemoteControlDMXExposedEntitiesGroupPatch::OnFixturePatchChanged(const UDMXEntityFixturePatch* FixturePatch)
+	{
+		RequestRefresh();
+	}
+
+	void SRemoteControlDMXExposedEntitiesGroupPatch::OnFixtureTypeChanged(const UDMXEntityFixtureType* FixtureType)
+	{
+		RequestRefresh();
+	}
+
+	void SRemoteControlDMXExposedEntitiesGroupPatch::OnPostLoadRemoteControlPreset(URemoteControlPreset* Preset)
+	{
+		RequestRefresh();
+	}
+
+	void SRemoteControlDMXExposedEntitiesGroupPatch::OnEntityExposedOrUnexposed(URemoteControlPreset* Preset, const FGuid& EntityId)
+	{
+		RequestRefresh();
+	}
+
+	void SRemoteControlDMXExposedEntitiesGroupPatch::OnEntityRebind(const FGuid& EntityId)
+	{
+		RequestRefresh();
+	}
+
+	void SRemoteControlDMXExposedEntitiesGroupPatch::OnEntitiesUpdated(URemoteControlPreset* Preset, const TSet<FGuid>& ModifiedEntities)
+	{
+		RequestRefresh();
+	}
+}
+
+#undef LOCTEXT_NAMESPACE
diff --git a/Engine/Plugins/VirtualProduction/RemoteControlProtocolDMX/Source/RemoteControlProtocolDMXEditor/Private/UI/SRemoteControlDMXExposedEntitiesGroupPatch.h b/Engine/Plugins/VirtualProduction/RemoteControlProtocolDMX/Source/RemoteControlProtocolDMXEditor/Private/UI/SRemoteControlDMXExposedEntitiesGroupPatch.h
new file mode 100644
--- /dev/null
+++ b/Engine/Plugins/VirtualProduction/RemoteControlProtocolDMX/Source/RemoteControlProtocolDMXEditor/Private/UI/SRemoteControlDMXExposedEntitiesGroupPatch.h
@@ -0,0 +1,94 @@
+// Copyright Epic Games, Inc. All Rights Reserved.
+
+#pragma once
+
+#include "Engine/TimerHandle.h"
+#include "Widgets/SCompoundWidget.h"
+
+class STableViewBase;
+class UDMXEntityFixturePatch;
+class UDMXEntityFixtureType;
+class URemoteControlPreset;
+class URemoteControlDMXUserData;
+enum class ERemoteControlDMXPatchGroupMode : uint8;
+struct FRemoteControlProperty;
+struct FRemoteControlProtocolEntity;
+template<typename T> class TStructOnScope;
+
+namespace UE::RemoteControl::DMX
+{
+	class FRemoteControlDMXFieldGroupColumnPatchItem;
+
+	class SRemoteControlDMXExposedEntitiesGroupPatch : public SCompoundWidget
+	{
+	public:
+		SLATE_BEGIN_ARGS(SRemoteControlDMXExposedEntitiesGroupPatch)
+			{}
+
+		SLATE_END_ARGS()
+
+		/** Constructs this widget */
+		void Construct(const FArguments& InArgs, const TWeakObjectPtr<URemoteControlPreset>& InPreset, const TArray<TSharedRef<FRemoteControlProperty>>& InChildProperties);
+
+	private:
+		/** Refreshes the widget */
+		void Refresh();
+
+		/** Refreshes the widget on the next tick */
+		void RequestRefresh();
+
+		/** Creates the menu */
+		TSharedRef<SWidget> CreateMenu();
+
+		/** Clears the patch for all entities */
+		void ClearPatch();
+
+		/** Generates a patch for all entities */
+		void GeneratePatch(); 
+
+		/** Enables or disables auto patching for the entities in this group */
+		void SetAutoPatchEnabled(bool bEnabled);
+
+		/** Called when a fixture patch was selected */
+		void OnFixturePatchSelected(TWeakObjectPtr<UDMXEntityFixturePatch> WeakFixturePatch);
+
+		/** Returns true if a primary fixture patch is selected */
+		bool IsPrimaryFixturePatch() const;
+
+		/** Returns the selected fixture patch or nullptr if no patch is selected */
+		UDMXEntityFixturePatch* GetFixturePatch() const;
+		
+		/** Returns the remote control DMX user data */
+		URemoteControlDMXUserData* GetDMXUserData() const;
+
+		/** Called when a fixture patch changed */
+		void OnFixturePatchChanged(const UDMXEntityFixturePatch* FixturePatch);
+
+		/** Called when a fixture type changed */
+		void OnFixtureTypeChanged(const UDMXEntityFixtureType* FixtureType);
+
+		/** Called when the remote control preset was fully loaded */
+		void OnPostLoadRemoteControlPreset(URemoteControlPreset* Preset);
+
+		/** Called when an entity was exposed or unexposed */
+		void OnEntityExposedOrUnexposed(URemoteControlPreset* Preset, const FGuid& EntityId);
+
+		/** Called when an entity was rebound */
+		void OnEntityRebind(const FGuid& EntityId);
+
+		/** Called when entities changed */
+		void OnEntitiesUpdated(URemoteControlPreset* Preset, const TSet<FGuid>& ModifiedEntities);
+
+		/** The entities that are affected by the patch */
+		TArray<TSharedRef<TStructOnScope<FRemoteControlProtocolEntity>>> Entities;
+
+		/** The preset for which this widget is displayed */
+		TWeakObjectPtr<URemoteControlPreset> WeakPreset;
+
+		/** The fixture patches in this preset */
+		TArray<TWeakObjectPtr<UDMXEntityFixturePatch>> WeakFixturePatches;
+
+		/** Timer handle for the request refresh method */
+		FTimerHandle RefreshTimerHandle;
+	};
+}
diff --git a/Engine/Plugins/VirtualProduction/RemoteControlProtocolDMX/Source/RemoteControlProtocolDMXEditor/Private/UI/SRemoteControlDMXExposedEntityPatch.cpp b/Engine/Plugins/VirtualProduction/RemoteControlProtocolDMX/Source/RemoteControlProtocolDMXEditor/Private/UI/SRemoteControlDMXExposedEntityPatch.cpp
new file mode 100644
--- /dev/null
+++ b/Engine/Plugins/VirtualProduction/RemoteControlProtocolDMX/Source/RemoteControlProtocolDMXEditor/Private/UI/SRemoteControlDMXExposedEntityPatch.cpp
@@ -0,0 +1,181 @@
+// Copyright Epic Games, Inc. All Rights Reserved.
+
+#include "SRemoteControlDMXExposedEntityPatch.h"
+
+#include "Editor.h"
+#include "Library/DMXEntityFixturePatch.h"
+#include "Library/RemoteControlDMXLibraryProxy.h"
+#include "RemoteControlDMXUserData.h"
+#include "RemoteControlPreset.h"
+#include "RemoteControlProtocolDMX.h"
+#include "ScopedTransaction.h"
+#include "Widgets/Input/SEditableTextBox.h"
+#include <limits>
+
+#define LOCTEXT_NAMESPACE "SRemoteControlDMXExposedEntityPatch"
+
+namespace UE::RemoteControl::DMX
+{
+	void SRemoteControlDMXExposedEntityPatch::Construct(const FArguments& InArgs, const TWeakObjectPtr<URemoteControlPreset>& InPreset, const TSharedRef<FRemoteControlProperty>& InProperty)
+	{
+		WeakPreset = InPreset;
+		if (!WeakPreset.IsValid())
+		{
+			return;
+		}
+		URemoteControlPreset* Preset = WeakPreset.Get();
+		check(Preset);
+
+		Property = InProperty;
+
+		URemoteControlDMXLibraryProxy::GetOnPostPropertyPatchesChanged().AddSP(this, &SRemoteControlDMXExposedEntityPatch::RequestRefresh);
+		Preset->OnEntityExposed().AddSP(this, &SRemoteControlDMXExposedEntityPatch::OnEntityExposedOrUnexposed);
+		Preset->OnEntityUnexposed().AddSP(this, &SRemoteControlDMXExposedEntityPatch::OnEntityExposedOrUnexposed);
+		Preset->OnEntityRebind().AddSP(this, &SRemoteControlDMXExposedEntityPatch::OnEntityRebind);
+		Preset->OnEntitiesUpdated().AddSP(this, &SRemoteControlDMXExposedEntityPatch::OnEntitiesUpdated);
+
+		RequestRefresh();
+	}
+
+	void SRemoteControlDMXExposedEntityPatch::RequestRefresh()
+	{
+		if (!RefreshTimerHandle.IsValid())
+		{
+			RefreshTimerHandle = GEditor->GetTimerManager()->SetTimerForNextTick(FTimerDelegate::CreateSP(this, &SRemoteControlDMXExposedEntityPatch::Refresh));
+		}
+	}
+
+	void SRemoteControlDMXExposedEntityPatch::Refresh()
+	{
+		RefreshTimerHandle.Invalidate();
+
+		TPair<TSharedPtr<TStructOnScope<FRemoteControlProtocolEntity>>, int64> MinEntityAbsoluteAddressPair = { nullptr, std::numeric_limits<int64>::max() };
+		for (const FRemoteControlProtocolBinding& Binding : Property->ProtocolBindings)
+		{
+			const TSharedPtr<TStructOnScope<FRemoteControlProtocolEntity>> Entity = Binding.GetRemoteControlProtocolEntityPtr();
+			const FRemoteControlDMXProtocolEntity* DMXEntity = Entity.IsValid() && Entity->IsValid() ? Entity->Cast<FRemoteControlDMXProtocolEntity>() : nullptr;
+			const UDMXEntityFixturePatch* FixturePatch = DMXEntity ? DMXEntity->ExtraSetting.FixturePatchReference.GetFixturePatch() : nullptr;
+			const FDMXFixtureMode* ActiveModePtr = FixturePatch ? FixturePatch->GetActiveMode() : nullptr;
+
+			if (FixturePatch && 
+				ActiveModePtr &&
+				ActiveModePtr->Functions.IsValidIndex(DMXEntity->ExtraSetting.FunctionIndex))
+			{
+				const FDMXFixtureFunction& Function = ActiveModePtr->Functions[DMXEntity->ExtraSetting.FunctionIndex];
+
+				const int32 FunctionStartingChannel = FixturePatch->GetStartingChannel() + Function.Channel - 1;
+				const int64 AbsoluteAddress = (int64)FixturePatch->GetUniverseID() * DMX_UNIVERSE_SIZE + FunctionStartingChannel;
+				const int64 PreviousMinAbsoluteAddress = MinEntityAbsoluteAddressPair.Value;
+				
+				if (PreviousMinAbsoluteAddress > AbsoluteAddress)
+				{
+					MinEntityAbsoluteAddressPair = TPair<TSharedPtr<TStructOnScope<FRemoteControlProtocolEntity>>, int64>(Entity, AbsoluteAddress);
+				}
+			}
+		}
+
+		if (!MinEntityAbsoluteAddressPair.Key.IsValid())
+		{
+			return;
+		}
+
+		MinEntity = MinEntityAbsoluteAddressPair.Key;
+
+		ChildSlot
+		[
+			SNew(SEditableTextBox)
+			.Font(FAppStyle::GetFontStyle("PropertyWindow.NormalFont"))
+			.Text(this, &SRemoteControlDMXExposedEntityPatch::GetUniverseChannelText)
+			.ToolTipText(LOCTEXT("UniverseChannelTooltip", "The universe and channel in the form of Universe.Channe. E.g. '1.1' is universe 1 channel 1."))
+			.IsReadOnly(this, &SRemoteControlDMXExposedEntityPatch::IsReadOnly)
+			.OnTextCommitted(this, &SRemoteControlDMXExposedEntityPatch::OnUniverseChannelTextCommitted)
+		];
+	}
+
+	FText SRemoteControlDMXExposedEntityPatch::GetUniverseChannelText() const
+	{
+		const FRemoteControlDMXProtocolEntity* DMXEntity = MinEntity.IsValid() && MinEntity->IsValid() ? MinEntity->Cast<FRemoteControlDMXProtocolEntity>() : nullptr;
+		const UDMXEntityFixturePatch* FixturePatch = DMXEntity ? DMXEntity->ExtraSetting.FixturePatchReference.GetFixturePatch() : nullptr;
+		const FDMXFixtureMode* ModePtr = FixturePatch ? FixturePatch->GetActiveMode() : nullptr;
+
+		if (ModePtr && ModePtr->Functions.IsValidIndex(DMXEntity->ExtraSetting.FunctionIndex))
+		{
+			const FDMXFixtureFunction& Function = ModePtr->Functions[DMXEntity->ExtraSetting.FunctionIndex];
+
+			const int32 FunctionStartingChannel = FixturePatch->GetStartingChannel() + Function.Channel - 1;
+			const FString UniverseChannelString = FString::Printf(TEXT("%i.%i"), FixturePatch->GetUniverseID(), FunctionStartingChannel);
+			return FText::FromString(UniverseChannelString);
+		}
+		else
+		{
+			return LOCTEXT("InvalidEntity", "Invalid Patch");
+		}
+	}
+
+	void SRemoteControlDMXExposedEntityPatch::OnUniverseChannelTextCommitted(const FText& InUniverseChannelText, ETextCommit::Type InCommitType)
+	{
+		const FRemoteControlDMXProtocolEntity* DMXEntity = MinEntity.IsValid() && MinEntity->IsValid() ? MinEntity->Cast<FRemoteControlDMXProtocolEntity>() : nullptr;
+		UDMXEntityFixturePatch* FixturePatch = DMXEntity ? DMXEntity->ExtraSetting.FixturePatchReference.GetFixturePatch() : nullptr;
+		if (!FixturePatch)
+		{
+			return;
+		}
+
+		TArray<FString> UniverseChannelStrings;
+		InUniverseChannelText.ToString().ParseIntoArray(UniverseChannelStrings, TEXT("."));
+
+		int32 Universe;
+		int32 Channel;
+		if (UniverseChannelStrings.Num() == 2 &&
+			LexTryParseString(Universe, *UniverseChannelStrings[0]) &&
+			LexTryParseString(Channel, *UniverseChannelStrings[1]))
+		{
+			const FScopedTransaction ReassignFixturePatchTransaction(LOCTEXT("ReassignFixturePatchTransaction", "Reassign Remote Control DMX Patch"));
+
+			FixturePatch->PreEditChange(nullptr);
+
+			FixturePatch->SetUniverseID(Universe);
+			FixturePatch->SetStartingChannel(Channel);
+
+			FixturePatch->PostEditChange();
+		}
+	}
+
+	bool SRemoteControlDMXExposedEntityPatch::IsReadOnly() const
+	{
+		const URemoteControlDMXUserData* DMXUserData = URemoteControlDMXUserData::GetOrCreateDMXUserData(WeakPreset.Get());
+		return DMXUserData ? DMXUserData->IsAutoPatch() : true;
+	}
+
+	void SRemoteControlDMXExposedEntityPatch::OnFixturePatchChanged(const UDMXEntityFixturePatch* FixturePatch)
+	{
+		RequestRefresh();
+	}
+
+	void SRemoteControlDMXExposedEntityPatch::OnFixtureTypeChanged(const UDMXEntityFixtureType* FixtureType)
+	{
+		RequestRefresh();
+	}
+
+	void SRemoteControlDMXExposedEntityPatch::OnPostLoadRemoteControlPreset(URemoteControlPreset* Preset)
+	{
+		RequestRefresh();
+	}
+
+	void SRemoteControlDMXExposedEntityPatch::OnEntityExposedOrUnexposed(URemoteControlPreset* Preset, const FGuid& EntityId)
+	{
+		RequestRefresh();
+	}
+
+	void SRemoteControlDMXExposedEntityPatch::OnEntityRebind(const FGuid& EntityId)
+	{
+		RequestRefresh();
+	}
+
+	void SRemoteControlDMXExposedEntityPatch::OnEntitiesUpdated(URemoteControlPreset* Preset, const TSet<FGuid>& ModifiedEntities)
+	{
+		RequestRefresh();
+	}
+}
+
+#undef LOCTEXT_NAMESPACE
diff --git a/Engine/Plugins/VirtualProduction/RemoteControlProtocolDMX/Source/RemoteControlProtocolDMXEditor/Private/UI/SRemoteControlDMXExposedEntityPatch.h b/Engine/Plugins/VirtualProduction/RemoteControlProtocolDMX/Source/RemoteControlProtocolDMXEditor/Private/UI/SRemoteControlDMXExposedEntityPatch.h
new file mode 100644
--- /dev/null
+++ b/Engine/Plugins/VirtualProduction/RemoteControlProtocolDMX/Source/RemoteControlProtocolDMXEditor/Private/UI/SRemoteControlDMXExposedEntityPatch.h
@@ -0,0 +1,76 @@
+// Copyright Epic Games, Inc. All Rights Reserved.
+
+#pragma once
+
+#include "Engine/TimerHandle.h"
+#include "Widgets/SCompoundWidget.h"
+
+class UDMXEntityFixturePatch;
+class UDMXEntityFixtureType;
+class URemoteControlDMXUserData;
+class URemoteControlPreset;
+enum class ERCFieldGroupType : uint8;
+struct FRemoteControlProperty;
+struct FRemoteControlProtocolEntity;
+template<typename T> class TStructOnScope;
+
+namespace UE::RemoteControl::DMX
+{
+	class SRemoteControlDMXExposedEntityPatch : public SCompoundWidget
+	{
+	public:
+		SLATE_BEGIN_ARGS(SRemoteControlDMXExposedEntityPatch)
+			{}
+
+		SLATE_END_ARGS()
+
+		/** Constructs this widget */
+		void Construct(const FArguments& InArgs, const TWeakObjectPtr<URemoteControlPreset>& InPreset, const TSharedRef<FRemoteControlProperty>& InProperty);
+
+	private:
+		/** Refreshes the widget on the next tick */
+		void RequestRefresh();
+
+		/** Refreshes the widget */
+		void Refresh();
+
+		/** Returns the current universe channel text */
+		FText GetUniverseChannelText() const;
+
+		/** Called when universe channel text was committed */
+		void OnUniverseChannelTextCommitted(const FText& InUniverseChannelText, ETextCommit::Type InCommitType);
+
+		/** Returns true if the text block of this widget is read only */
+		bool IsReadOnly() const;
+
+		/** Called when a fixture patch changed */
+		void OnFixturePatchChanged(const UDMXEntityFixturePatch* FixturePatch);
+
+		/** Called when a fixture type changed */
+		void OnFixtureTypeChanged(const UDMXEntityFixtureType* FixtureType);
+
+		/** Called when the remote control preset was fully loaded */
+		void OnPostLoadRemoteControlPreset(URemoteControlPreset* Preset);
+
+		/** Called when an entity was exposed or unexposed */
+		void OnEntityExposedOrUnexposed(URemoteControlPreset* Preset, const FGuid& EntityId);
+
+		/** Called when an entity was rebound */
+		void OnEntityRebind(const FGuid& EntityId);
+
+		/** Called when entities changed */
+		void OnEntitiesUpdated(URemoteControlPreset* Preset, const TSet<FGuid>& ModifiedEntities);
+
+		/** The displayed property */
+		TSharedPtr<FRemoteControlProperty> Property;
+		
+		/** The first binding given the DMX universe and channel */
+		TSharedPtr<TStructOnScope<FRemoteControlProtocolEntity>> MinEntity;
+
+		/** The preset for which this widget is painted */
+		TWeakObjectPtr<URemoteControlPreset> WeakPreset;
+
+		/** Timer handle for request refresh */
+		FTimerHandle RefreshTimerHandle;
+	};
+}
diff --git a/Engine/Plugins/VirtualProduction/RemoteControlProtocolDMX/Source/RemoteControlProtocolDMXEditor/Private/UI/SRemoteControlDMXPresetUserData.cpp b/Engine/Plugins/VirtualProduction/RemoteControlProtocolDMX/Source/RemoteControlProtocolDMXEditor/Private/UI/SRemoteControlDMXPresetUserData.cpp
new file mode 100644
--- /dev/null
+++ b/Engine/Plugins/VirtualProduction/RemoteControlProtocolDMX/Source/RemoteControlProtocolDMXEditor/Private/UI/SRemoteControlDMXPresetUserData.cpp
@@ -0,0 +1,258 @@
+// Copyright Epic Games, Inc. All Rights Reserved.
+
+#include "SRemoteControlDMXPresetUserData.h"
+
+#include "Customizations/RemoteControlProtocolDMXPresetUserDataDetails.h"
+#include "DetailsViewArgs.h"
+#include "DMXEditorModule.h"
+#include "DMXEditorStyle.h"
+#include "IDetailsView.h"
+#include "IRCProtocolBindingList.h"
+#include "IRemoteControlProtocolWidgetsModule.h"
+#include "Library/DMXLibrary.h"
+#include "Modules/ModuleManager.h"
+#include "PropertyEditorModule.h"
+#include "RemoteControlDMXUserData.h"
+#include "RemoteControlPreset.h"
+#include "ScopedTransaction.h"
+#include "Styling/AppStyle.h"
+#include "ToolMenus.h"
+#include "Widgets/Input/SCheckBox.h"
+#include "Widgets/Input/SEditableTextBox.h"
+#include "Widgets/Layout/SSeparator.h"
+
+#define LOCTEXT_NAMESPACE "SRemoteControlDMXPresetUserData"
+
+namespace UE::RemoteControl::DMX
+{
+	void SRemoteControlDMXPresetUserData::Construct(const FArguments& InArgs, URemoteControlDMXUserData* InDMXUserData)
+	{
+		if (!ensureMsgf(InDMXUserData, TEXT("%hs: Invalid Remote Control DMX User Data provided, cannot draw widget."), __FUNCTION__))
+		{
+			return;
+		}
+		DMXUserData = InDMXUserData;
+		
+		// Create a details view for the user data to present the DMX library
+		FDetailsViewArgs DetailsViewArgs;
+		DetailsViewArgs.bAllowSearch = false;
+		DetailsViewArgs.NameAreaSettings = FDetailsViewArgs::HideNameArea;
+		DetailsViewArgs.bShowOptions = false;
+
+		FPropertyEditorModule& PropertyEditorModule = FModuleManager::GetModuleChecked<FPropertyEditorModule>("PropertyEditor");
+		const TSharedRef<IDetailsView> DetailsView = PropertyEditorModule.CreateDetailView(DetailsViewArgs);
+		
+		DetailsView->RegisterInstancedCustomPropertyLayout(URemoteControlDMXUserData::StaticClass(), 
+			FOnGetDetailCustomizationInstance::CreateStatic(&FRemoteControlProtocolDMXPresetUserDataDetails::MakeInstance));
+		
+		DetailsView->SetObject(InDMXUserData);
+
+		ChildSlot
+			[
+				SNew(SHorizontalBox)
+
+				// Auto Patch option
+				+ SHorizontalBox::Slot()
+				.Padding(FMargin(8.f, 0.f))
+				.VAlign(VAlign_Center)
+				.AutoWidth()
+				[
+					SNew(SCheckBox)
+					.IsChecked(this, &SRemoteControlDMXPresetUserData::GetAutoPatchCheckState)
+					.OnCheckStateChanged(this, &SRemoteControlDMXPresetUserData::OnAutoPatchCheckStateChanged)
+					.ToolTipText(LOCTEXT("AutoPatchCheckBoxTooltip", "Enables auto assign patches. Note, in auto assign mode patches are created depening on the sort order.."))
+					[
+						SNew(SBorder)
+						.BorderImage(FAppStyle::GetBrush("NoBorder"))
+						.Padding(FMargin(8.f, 0.f))
+						[
+							SNew(STextBlock)
+							.ColorAndOpacity(FSlateColor::UseSubduedForeground())
+							.Text(LOCTEXT("AutoPatchCheckBoxLabel", "Auto Patch in List Order"))
+						]
+					]
+				]
+					
+				+ SHorizontalBox::Slot()
+				.AutoWidth()
+				[
+					SNew(SSeparator)
+					.Orientation(Orient_Vertical)
+				]
+
+				// Auto assign from universe
+				+ SHorizontalBox::Slot()
+				.Padding(FMargin(8.f, 0.f))
+				.AutoWidth()
+				[
+					SNew(SBorder)
+					.BorderImage(FAppStyle::GetBrush("NoBorder"))
+					.Padding(FMargin(4.f, 0.f))
+					.Visibility(this, &SRemoteControlDMXPresetUserData::GetAutoAssignFromUniverseVisibility)
+					[
+						SNew(SHorizontalBox)
+
+						// Label
+						+ SHorizontalBox::Slot()
+						.Padding(FMargin(4.f, 0.f))
+						.VAlign(VAlign_Center)
+						.AutoWidth()
+						[
+							SNew(STextBlock)
+							.Text(LOCTEXT("AutoAssignFromUniverseLabel", "Patch from Universe"))
+						]
+
+						// Editable text block
+						+ SHorizontalBox::Slot()
+						.VAlign(VAlign_Center)
+						.Padding(FMargin(4.f, 0.f))
+						.AutoWidth()
+						[
+							SAssignNew(AutoAssignFromUniverseEditableTextBox, SEditableTextBox)
+							.MinDesiredWidth(40.f)
+							.TextFlowDirection(ETextFlowDirection::RightToLeft)
+							.Text(this, &SRemoteControlDMXPresetUserData::GetAutoAssignFromUniverseText)
+							.OnTextCommitted(this, &SRemoteControlDMXPresetUserData::OnAutoAssignFromUniverseTextCommitted)
+						]
+					]
+				]
+
+				+ SHorizontalBox::Slot()
+				.AutoWidth()
+				[
+					SNew(SSeparator)
+					.Orientation(Orient_Vertical)
+				]
+
+				// Export as MVR button
+				+ SHorizontalBox::Slot()
+				.Padding(FMargin(8.f, 0.f))
+				.VAlign(VAlign_Center)
+				.FillWidth(1.f)
+				[
+					GenerateExportAsMVRButton()
+				]
+
+					
+				// DMXLibrary
+				+ SHorizontalBox::Slot()
+				.Padding(FMargin(8.f, 0.f))
+				.VAlign(VAlign_Center)	
+				.AutoWidth()
+				[
+					DetailsView
+				]
+			];
+	}
+
+	void SRemoteControlDMXPresetUserData::AddReferencedObjects(FReferenceCollector& Collector)
+	{
+		Collector.AddReferencedObject(DMXUserData);
+	}
+	
+	FString SRemoteControlDMXPresetUserData::GetReferencerName() const
+	{
+		return TEXT("SRemoteControlDMXPresetUserData");
+	}
+	
+	TSharedRef<SWidget> SRemoteControlDMXPresetUserData::GenerateExportAsMVRButton()
+	{
+		constexpr const TCHAR* MenuName = TEXT("RemoteControlPresetDMXUserDataMenu");
+
+		UToolMenus* ToolMenus = UToolMenus::Get();
+		if (!ToolMenus->IsMenuRegistered(MenuName))
+		{
+			const FName NoParentToolbar = NAME_None;
+			const FName NoNameOverride = NAME_None;
+			const FText NoLabel;
+
+			UToolMenu* Menu = ToolMenus->RegisterMenu(MenuName, NoParentToolbar, EMultiBoxType::SlimHorizontalToolBar);
+			FToolMenuSection& Section = Menu->AddSection("DMXPresetUserDataMenu", NoLabel);
+
+			Section.AddMenuEntry
+			(
+				NoNameOverride,
+				LOCTEXT("ExportAsMVRLabel", "Export as MVR"),
+				LOCTEXT("ExportAsMVRTooltip", "Exports the Remote Control DMX Library as MVR file"),
+				FSlateIcon(FDMXEditorStyle::Get().GetStyleSetName(), "Icons.DMXLibraryToolbar.Export"),
+				FUIAction(FExecuteAction::CreateSP(this, &SRemoteControlDMXPresetUserData::OnExportAsMVRClicked)),
+				EUserInterfaceActionType::Button
+			);
+		}
+
+		const FToolMenuContext Context(DMXUserData);
+		UToolMenu* Menu = ToolMenus->GenerateMenu(MenuName, Context);
+
+		return ToolMenus->GenerateWidget(Menu);
+	}
+
+	void SRemoteControlDMXPresetUserData::OnExportAsMVRClicked()
+	{
+		UDMXLibrary* DMXLibrary = DMXUserData ? DMXUserData->GetDMXLibrary() : nullptr;
+		if (DMXLibrary)
+		{
+			URemoteControlPreset* RemoteControlPreset = Cast<URemoteControlPreset>(DMXUserData->GetOuter());
+			const FString DesiredFileName = RemoteControlPreset ? RemoteControlPreset->GetName() : FString();
+
+			const FDMXEditorModule& DMXEditorModule = FModuleManager::GetModuleChecked<FDMXEditorModule>("DMXEditor");
+			DMXEditorModule.ExportDMXLibraryAsMVRFile(DMXLibrary, DesiredFileName);
+		}
+	}
+
+	ECheckBoxState SRemoteControlDMXPresetUserData::GetAutoPatchCheckState() const
+	{
+		return DMXUserData && DMXUserData->IsAutoPatch() ? ECheckBoxState::Checked : ECheckBoxState::Unchecked;
+	}
+
+	void SRemoteControlDMXPresetUserData::OnAutoPatchCheckStateChanged(ECheckBoxState NewCheckState)
+	{
+		const bool bAutoPatch = NewCheckState == ECheckBoxState::Checked;
+		DMXUserData->SetAutoPatchEnabled(bAutoPatch);
+
+		const IRemoteControlProtocolWidgetsModule& RCWidgetsModule = IRemoteControlProtocolWidgetsModule::Get();
+		const TSharedPtr<IRCProtocolBindingList> BindingList = RCWidgetsModule.GetProtocolBindingList();
+		if (BindingList.IsValid())
+		{
+			BindingList->Refresh();
+		}
+	}
+
+	EVisibility SRemoteControlDMXPresetUserData::GetAutoAssignFromUniverseVisibility() const
+	{
+		return DMXUserData && DMXUserData->IsAutoPatch() ? EVisibility::Visible : EVisibility::Collapsed;
+	}
+
+	FText SRemoteControlDMXPresetUserData::GetAutoAssignFromUniverseText() const
+	{
+		if (DMXUserData)
+		{
+			const int32 AutoAssignFromUniverse = DMXUserData->GetAutoAssignFromUniverse();
+			const FString AutoAssignFromUniverseString = FString::FromInt(AutoAssignFromUniverse);
+			return FText::FromString(AutoAssignFromUniverseString);
+		}
+
+		return FText::GetEmpty();
+	}
+
+	void SRemoteControlDMXPresetUserData::OnAutoAssignFromUniverseTextCommitted(const FText& InAutoAssignFromUniverseText, ETextCommit::Type InCommitType)
+	{
+		const FScopedTransaction ReassignFixturePatchTransaction(LOCTEXT("ReassignFixturePatchesTransaction", "Set Patch To DMX Universe"));
+
+		int32 AutoAssignFromUniverse = 0;
+		if (DMXUserData &&
+			LexTryParseString(AutoAssignFromUniverse, *InAutoAssignFromUniverseText.ToString()) &&
+			AutoAssignFromUniverse > 0)
+		{
+			DMXUserData->PreEditChange(nullptr);
+			DMXUserData->SetAutoAssignFromUniverse(AutoAssignFromUniverse);
+			DMXUserData->PostEditChange();
+		}
+		else if (DMXUserData)
+		{
+			const FText AutoAssignFromUniverseText = FText::FromString(FString::FromInt(DMXUserData->GetAutoAssignFromUniverse()));
+			AutoAssignFromUniverseEditableTextBox->SetText(AutoAssignFromUniverseText);
+		}
+	}
+}
+
+#undef LOCTEXT_NAMESPACE
diff --git a/Engine/Plugins/VirtualProduction/RemoteControlProtocolDMX/Source/RemoteControlProtocolDMXEditor/Private/UI/SRemoteControlDMXPresetUserData.h b/Engine/Plugins/VirtualProduction/RemoteControlProtocolDMX/Source/RemoteControlProtocolDMXEditor/Private/UI/SRemoteControlDMXPresetUserData.h
new file mode 100644
--- /dev/null
+++ b/Engine/Plugins/VirtualProduction/RemoteControlProtocolDMX/Source/RemoteControlProtocolDMXEditor/Private/UI/SRemoteControlDMXPresetUserData.h
@@ -0,0 +1,60 @@
+// Copyright Epic Games, Inc. All Rights Reserved.
+
+#pragma once
+
+#include "UObject/GCObject.h"
+#include "Widgets/SCompoundWidget.h"
+
+class SEditableTextBox;
+class URemoteControlDMXUserData;
+enum class ECheckBoxState : uint8;
+
+namespace UE::RemoteControl::DMX
+{
+	class SRemoteControlDMXPresetUserData : public SCompoundWidget, public FGCObject
+	{
+	public:
+		SLATE_BEGIN_ARGS(SRemoteControlDMXPresetUserData) 
+			{}
+
+		SLATE_END_ARGS()
+
+		/** Constructs this widget */
+		void Construct(const FArguments& InArgs, URemoteControlDMXUserData* InDMXUserData);
+
+	protected:
+		// FGCObject interface
+		virtual void AddReferencedObjects(FReferenceCollector& Collector) override;
+		virtual FString GetReferencerName() const override;
+		// End FGCObject interface
+
+	private:
+		/** Generates the Export as MVR button */
+		TSharedRef<SWidget> GenerateExportAsMVRButton();
+
+		/** Called when the export as MVR button was clicked */
+		void OnExportAsMVRClicked();
+
+		/** Returns true if the auto patch option is checked */
+		ECheckBoxState GetAutoPatchCheckState() const;
+
+		/** Called when the auto patch check state changed */
+		void OnAutoPatchCheckStateChanged(ECheckBoxState NewCheckState);
+		
+		/** Returns the visibility of the auto assign from universe option */
+		EVisibility GetAutoAssignFromUniverseVisibility() const;
+
+		/** Returns the auto assign from universe as text */
+		FText GetAutoAssignFromUniverseText() const;
+
+		/** Called when the auto assign from universe text was committed */
+		void OnAutoAssignFromUniverseTextCommitted(const FText& InAutoAssignFromUniverseText, ETextCommit::Type InCommitType);
+		
+		/** Text box to edit the auto assign from universe */
+		TSharedPtr<SEditableTextBox> AutoAssignFromUniverseEditableTextBox;
+
+		/** The preset this widget is editing */
+		TObjectPtr<URemoteControlDMXUserData> DMXUserData;
+	};
+}
+
diff --git a/Engine/Plugins/VirtualProduction/RemoteControlProtocolDMX/Source/RemoteControlProtocolDMXEditor/RemoteControlProtocolDMXEditor.Build.cs b/Engine/Plugins/VirtualProduction/RemoteControlProtocolDMX/Source/RemoteControlProtocolDMXEditor/RemoteControlProtocolDMXEditor.Build.cs
--- a/Engine/Plugins/VirtualProduction/RemoteControlProtocolDMX/Source/RemoteControlProtocolDMXEditor/RemoteControlProtocolDMXEditor.Build.cs
+++ b/Engine/Plugins/VirtualProduction/RemoteControlProtocolDMX/Source/RemoteControlProtocolDMXEditor/RemoteControlProtocolDMXEditor.Build.cs
@@ -11,21 +11,27 @@ public class RemoteControlProtocolDMXEditor : ModuleRules
 		PublicDependencyModuleNames.AddRange(
 			new string[] {}
 		);
-
+		
         PrivateDependencyModuleNames.AddRange(
 			new string[] {
 				"Core",
 				"CoreUObject",
+				"DMXEditor",
 				"DMXProtocol",
+				"DMXRuntime",
 				"DMXProtocolEditor",
-				
+				"Engine",
 				"InputCore",
 				"PropertyEditor",
 				"RemoteControl",
 				"RemoteControlProtocol",
 				"RemoteControlProtocolDMX",
+				"RemoteControlUI",
+				"RemoteControlProtocolWidgets",
 				"Slate",
-				"SlateCore"
+				"SlateCore",
+				"ToolMenus",
+				"UnrealEd"
 			}
 		);
     }
