DMX - Fix Static Analysis, identical sub-expressions to the left and to the right of the '==' operator

diff --git a/Engine/Plugins/VirtualProduction/RemoteControlProtocolDMX/Source/RemoteControlProtocolDMX/Private/Library/RemoteControlDMXControlledProperty.cpp b/Engine/Plugins/VirtualProduction/RemoteControlProtocolDMX/Source/RemoteControlProtocolDMX/Private/Library/RemoteControlDMXControlledProperty.cpp
--- a/Engine/Plugins/VirtualProduction/RemoteControlProtocolDMX/Source/RemoteControlProtocolDMX/Private/Library/RemoteControlDMXControlledProperty.cpp
+++ b/Engine/Plugins/VirtualProduction/RemoteControlProtocolDMX/Source/RemoteControlProtocolDMX/Private/Library/RemoteControlDMXControlledProperty.cpp
@@ -109,7 +109,7 @@ namespace UE::RemoteControl::DMX
 		for (const TSharedRef<TStructOnScope<FRemoteControlProtocolEntity>>& Entity : Entities)
 		{
 			FRemoteControlDMXProtocolEntity* DMXEntity = Entity->IsValid() ? Entity->Cast<FRemoteControlDMXProtocolEntity>() : nullptr;
-			if (!DMXEntity || FirstDMXEntity == FirstDMXEntity)
+			if (!DMXEntity || DMXEntity == FirstDMXEntity)
 			{
 				continue;
 			}
