Remote Control DMX: Fix an issue where DMX Controlled Properties no longer read the correct function when first being loaded

#rnx

diff --git a/Engine/Plugins/VirtualProduction/RemoteControlProtocolDMX/Source/RemoteControlProtocolDMX/Private/Library/RemoteControlDMXControlledProperty.cpp b/Engine/Plugins/VirtualProduction/RemoteControlProtocolDMX/Source/RemoteControlProtocolDMX/Private/Library/RemoteControlDMXControlledProperty.cpp
--- a/Engine/Plugins/VirtualProduction/RemoteControlProtocolDMX/Source/RemoteControlProtocolDMX/Private/Library/RemoteControlDMXControlledProperty.cpp
+++ b/Engine/Plugins/VirtualProduction/RemoteControlProtocolDMX/Source/RemoteControlProtocolDMX/Private/Library/RemoteControlDMXControlledProperty.cpp
@@ -103,7 +103,6 @@ namespace UE::RemoteControl::DMX
 		}
 
 		UDMXEntityFixturePatch* UnifiedFixturePatch = FirstDMXEntity->ExtraSetting.FixturePatchReference.GetFixturePatch();
-		const int32 UnifiedFunctionIndex = FirstDMXEntity->ExtraSetting.FunctionIndex;
 		const bool bUnifiedIsPrimaryPatch = FirstDMXEntity->ExtraSetting.bIsPrimaryPatch;
 
 		for (const TSharedRef<TStructOnScope<FRemoteControlProtocolEntity>>& Entity : Entities)
@@ -116,13 +115,11 @@ namespace UE::RemoteControl::DMX
 
 			// Only update if the properties differ, as this will trigger another rebuild of the library proxy that uses this object
 			if (DMXEntity->ExtraSetting.FixturePatchReference.GetFixturePatch() != UnifiedFixturePatch ||
-				DMXEntity->ExtraSetting.FunctionIndex != UnifiedFunctionIndex ||
 				DMXEntity->ExtraSetting.bIsPrimaryPatch != bUnifiedIsPrimaryPatch)
 			{
 				ensureMsgf(0, TEXT("Found entities of the same property with differing properties. Adopting property from first entity to mend."));
 
 				DMXEntity->ExtraSetting.FixturePatchReference = UnifiedFixturePatch;
-				DMXEntity->ExtraSetting.FunctionIndex = UnifiedFunctionIndex;
 				DMXEntity->ExtraSetting.bIsPrimaryPatch = bUnifiedIsPrimaryPatch;
 			}
 		}
