DMX: DMXEntityFixtureType::CreateFixtureTypeInDMXLibrary now correctly updates the Channel Span for its DMX Modes

diff --git a/Engine/Plugins/VirtualProduction/DMX/DMXEngine/Source/DMXRuntime/Private/Library/DMXEntityFixtureType.cpp b/Engine/Plugins/VirtualProduction/DMX/DMXEngine/Source/DMXRuntime/Private/Library/DMXEntityFixtureType.cpp
--- a/Engine/Plugins/VirtualProduction/DMX/DMXEngine/Source/DMXRuntime/Private/Library/DMXEntityFixtureType.cpp
+++ b/Engine/Plugins/VirtualProduction/DMX/DMXEngine/Source/DMXRuntime/Private/Library/DMXEntityFixtureType.cpp
@@ -182,6 +182,13 @@ UDMXEntityFixtureType* UDMXEntityFixtureType::CreateFixtureTypeInLibrary(FDMXEnt
 		NewFixtureType->DMXCategory = ConstructionParams.DMXCategory;
 		NewFixtureType->Modes = ConstructionParams.Modes;
 
+		// Update the channel span for each mode
+		for (int32 ModeIndex = 0; ModeIndex < NewFixtureType->Modes.Num(); ModeIndex++)
+		{
+			NewFixtureType->UpdateChannelSpan(ModeIndex);
+		}
+
+
 #if WITH_EDITOR
 		if (bMarkDMXLibraryDirty)
 		{
diff --git a/Engine/Plugins/VirtualProduction/RemoteControlProtocolDMX/Source/RemoteControlProtocolDMXEditor/Private/Library/RemoteControlDMXPatchBuilder.cpp b/Engine/Plugins/VirtualProduction/RemoteControlProtocolDMX/Source/RemoteControlProtocolDMXEditor/Private/Library/RemoteControlDMXPatchBuilder.cpp
--- a/Engine/Plugins/VirtualProduction/RemoteControlProtocolDMX/Source/RemoteControlProtocolDMXEditor/Private/Library/RemoteControlDMXPatchBuilder.cpp
+++ b/Engine/Plugins/VirtualProduction/RemoteControlProtocolDMX/Source/RemoteControlProtocolDMXEditor/Private/Library/RemoteControlDMXPatchBuilder.cpp
@@ -306,31 +306,9 @@ namespace UE::RemoteControl::DMX
 
 		void FRCSinglePatchBuilder::UpdateFixtureType(UDMXLibrary& InDMXLibrary, UDMXEntityFixtureType*& InOutFixtureType)
 		{
-			if (!InOutFixtureType)
-			{
-				// Create a new fixture type
-				FDMXEntityFixtureTypeConstructionParams FixtureTypeConstructionParams;
-				FixtureTypeConstructionParams.DMXCategory = FDMXFixtureCategory(TEXT("Remote Control"));
-				FixtureTypeConstructionParams.ParentDMXLibrary = &InDMXLibrary;
-
-				InOutFixtureType = UDMXEntityFixtureType::CreateFixtureTypeInLibrary(FixtureTypeConstructionParams);
-				check(InOutFixtureType);
-
-				const FString DesiredFixtureTypeName = TEXT("FT_") + GetDesiredName();
-				InOutFixtureType->Name = DesiredFixtureTypeName;
-			}
-
-			if (!InOutFixtureType)
-			{
-				return;
-			}
-
-			// Rebuild modes
-			InOutFixtureType->Modes.Reset();
-			InOutFixtureType->Modes.Add(FDMXFixtureMode());
-
-			FDMXFixtureMode& Mode = InOutFixtureType->Modes.Last();
-			Mode.ModeName = TEXT("RemoteControl");
+			// Build modes
+			FDMXFixtureMode NewMode;
+			NewMode.ModeName = TEXT("RemoteControl");
 
 			for (const TSharedRef<FRemoteControlDMXControlledProperty>& Property : DMXControlledProperties)
 			{
@@ -342,9 +320,9 @@ namespace UE::RemoteControl::DMX
 						continue;
 					}
 
-					const int32 NextFreeChannel = [&Mode]()
+					const int32 NextFreeChannel = [&NewMode]()
 						{
-							const FDMXFixtureFunction* MaxFunctionPtr = Algo::MaxElementBy(Mode.Functions, &FDMXFixtureFunction::Channel);
+							const FDMXFixtureFunction* MaxFunctionPtr = Algo::MaxElementBy(NewMode.Functions, &FDMXFixtureFunction::Channel);
 
 							return MaxFunctionPtr ? MaxFunctionPtr->GetLastChannel() + 1 : 1;
 						}();
@@ -359,12 +337,31 @@ namespace UE::RemoteControl::DMX
 						// Addopt the attribute name from the DMX entity
 						NewFunction.Attribute = DMXEntity->ExtraSetting.AttributeName;
 
-						const int32 FunctionIndex = Mode.Functions.Add(NewFunction);
+						const int32 FunctionIndex = NewMode.Functions.Add(NewFunction);
 
 						// Remember the function index from the fixture function
 						DMXEntity->ExtraSetting.FunctionIndex = FunctionIndex;
 				}
 			}
+
+			// Get or create the fixture type
+			if (!InOutFixtureType)
+			{
+				// Create a new fixture type
+				FDMXEntityFixtureTypeConstructionParams FixtureTypeConstructionParams;
+				FixtureTypeConstructionParams.DMXCategory = FDMXFixtureCategory(TEXT("Remote Control"));
+				FixtureTypeConstructionParams.ParentDMXLibrary = &InDMXLibrary;
+				FixtureTypeConstructionParams.Modes = { NewMode };
+
+				InOutFixtureType = UDMXEntityFixtureType::CreateFixtureTypeInLibrary(FixtureTypeConstructionParams);
+				check(InOutFixtureType);
+
+				const FString DesiredFixtureTypeName = TEXT("FT_") + GetDesiredName();
+				InOutFixtureType->Name = DesiredFixtureTypeName;
+			}
+
+			constexpr int32 RCModeIndex = 0;
+			InOutFixtureType->UpdateChannelSpan(RCModeIndex);
 		}
 
 		UDMXEntityFixturePatch* FRCSinglePatchBuilder::FindFixturePatch(UDMXEntityFixtureType* FixtureType) const
