Remote Control DMX: Fix various issues and add small improvements to the initial DMX Library support implementation

* Fix an issue where the patch was reset when changing the level
* Fix an issue where property changes of bindings was not reflected in the DMX Library
* Fix an issue where properties were not always reflected correctly in the fixture type
* Entities list for protocols now shows grouped by owner by default
* Change attribute names to show the controlled property name before the subobject path

#rnx

diff --git a/Engine/Plugins/VirtualProduction/RemoteControlProtocolDMX/Source/RemoteControlProtocolDMX/Internal/Library/RemoteControlDMXControlledPropertyPatch.h b/Engine/Plugins/VirtualProduction/RemoteControlProtocolDMX/Source/RemoteControlProtocolDMX/Internal/Library/RemoteControlDMXControlledPropertyPatch.h
--- a/Engine/Plugins/VirtualProduction/RemoteControlProtocolDMX/Source/RemoteControlProtocolDMX/Internal/Library/RemoteControlDMXControlledPropertyPatch.h
+++ b/Engine/Plugins/VirtualProduction/RemoteControlProtocolDMX/Source/RemoteControlProtocolDMX/Internal/Library/RemoteControlDMXControlledPropertyPatch.h
@@ -35,6 +35,9 @@ namespace UE::RemoteControl::DMX
 		/** Returns the fixture patch corresponding to this property patch */
 		UDMXEntityFixturePatch* GetFixturePatch() const;
 
+		/** Returns the owner object of this patch */
+		const UObject* GetOwnerObject() const;
+
 	private:
 #if WITH_EDITOR
 		/** Called when a fixture type changed. Useful to adopt changed properties from the fixture type. */
diff --git a/Engine/Plugins/VirtualProduction/RemoteControlProtocolDMX/Source/RemoteControlProtocolDMX/Internal/Library/RemoteControlDMXLibraryProxy.h b/Engine/Plugins/VirtualProduction/RemoteControlProtocolDMX/Source/RemoteControlProtocolDMX/Internal/Library/RemoteControlDMXLibraryProxy.h
--- a/Engine/Plugins/VirtualProduction/RemoteControlProtocolDMX/Source/RemoteControlProtocolDMX/Internal/Library/RemoteControlDMXLibraryProxy.h
+++ b/Engine/Plugins/VirtualProduction/RemoteControlProtocolDMX/Source/RemoteControlProtocolDMX/Internal/Library/RemoteControlDMXLibraryProxy.h
@@ -15,6 +15,7 @@ namespace UE::RemoteControl::DMX
 { 
 	class FRemoteControlDMXControlledProperty; 
 	class FRemoteControlDMXControlledPropertyPatch;
+	class FRemoteControlDMXProtocolEntityObserver;
 }
 
 DECLARE_MULTICAST_DELEGATE_OneParam(FRemoteControlDMXPrePropertyPatchesChanged, URemoteControlPreset* /** ChangedPreset */)
@@ -27,6 +28,7 @@ class REMOTECONTROLPROTOCOLDMX_API URemoteControlDMXLibraryProxy : public UObjec
 
 	using FRemoteControlDMXControlledProperty = UE::RemoteControl::DMX::FRemoteControlDMXControlledProperty;
 	using FRemoteControlDMXControlledPropertyPatch = UE::RemoteControl::DMX::FRemoteControlDMXControlledPropertyPatch;
+	using FRemoteControlDMXProtocolEntityObserver = UE::RemoteControl::DMX::FRemoteControlDMXProtocolEntityObserver;
 
 public:
 	//~ Begin UObject interface
@@ -70,6 +72,11 @@ private:
 	/** Unbinds from the OnFixturePatchReceived event for all patches in use */
 	void UnbindOnFixturePatchesReceived();
 
+#if WITH_EDITOR
+	/** Updates the Entities Observer */
+	void UpdateEntitiesObserver();
+#endif
+
 	/** Called when a fixture patch was received */
 	UFUNCTION()
 	void OnFixturePatchReceived(UDMXEntityFixturePatch* FixturePatch, const FDMXNormalizedAttributeValueMap& ValuePerAttribute);
@@ -98,6 +105,9 @@ private:
 	/** Called when entities changed */
 	void OnEntitiesUpdated(URemoteControlPreset* Preset, const TSet<FGuid>& ModifiedEntities);
 
+	/** Called when a map was fully loaded */
+	void OnPostLoadMapWithWorld(UWorld* World);
+
 	/** Current DMX controlled property patches */
 	TArray<TSharedRef<FRemoteControlDMXControlledPropertyPatch>> PropertyPatches;
 
@@ -105,6 +115,9 @@ private:
 	FDelegateHandle RefreshDelegateHandle;
 
 #if WITH_EDITOR
+	/** Observes entites for DMX specific property changes in editor */
+	TSharedPtr<FRemoteControlDMXProtocolEntityObserver> EntitiesObserver;
+
 	/** Delegate broadcast before property patches are being changed */
 	static FRemoteControlDMXPrePropertyPatchesChanged OnPrePropertyPatchesChanged;
 
diff --git a/Engine/Plugins/VirtualProduction/RemoteControlProtocolDMX/Source/RemoteControlProtocolDMX/Internal/RemoteControlDMXUserData.h b/Engine/Plugins/VirtualProduction/RemoteControlProtocolDMX/Source/RemoteControlProtocolDMX/Internal/RemoteControlDMXUserData.h
--- a/Engine/Plugins/VirtualProduction/RemoteControlProtocolDMX/Source/RemoteControlProtocolDMX/Internal/RemoteControlDMXUserData.h
+++ b/Engine/Plugins/VirtualProduction/RemoteControlProtocolDMX/Source/RemoteControlProtocolDMX/Internal/RemoteControlDMXUserData.h
@@ -93,7 +93,7 @@ private:
 
 	/** Defines how DMX Protocol Entities should be patched  */
 	UPROPERTY()
-	ERemoteControlDMXPatchGroupMode PatchGroupMode;
+	ERemoteControlDMXPatchGroupMode PatchGroupMode = ERemoteControlDMXPatchGroupMode::GroupByOwner;
 
 	/** Proxy to handle the DMX Library */
 	UPROPERTY(Transient)
diff --git a/Engine/Plugins/VirtualProduction/RemoteControlProtocolDMX/Source/RemoteControlProtocolDMX/Private/Library/RemoteControlDMXControlledProperty.cpp b/Engine/Plugins/VirtualProduction/RemoteControlProtocolDMX/Source/RemoteControlProtocolDMX/Private/Library/RemoteControlDMXControlledProperty.cpp
--- a/Engine/Plugins/VirtualProduction/RemoteControlProtocolDMX/Source/RemoteControlProtocolDMX/Private/Library/RemoteControlDMXControlledProperty.cpp
+++ b/Engine/Plugins/VirtualProduction/RemoteControlProtocolDMX/Source/RemoteControlProtocolDMX/Private/Library/RemoteControlDMXControlledProperty.cpp
@@ -80,11 +80,11 @@ namespace UE::RemoteControl::DMX
 					{
 						if (Entities.Num() == 1)
 						{
-							return FString::Printf(TEXT("%s.%s"), *GetSubobjectPath(), *ExposedProperty->FieldPathInfo.ToString());
+							return FString::Printf(TEXT("%s.%s"), *ExposedProperty->FieldPathInfo.ToString(), *GetSubobjectPath());
 						}
 						else
 						{
-							return FString::Printf(TEXT("%s.%s%i"), *GetSubobjectPath(), *ExposedProperty->FieldPathInfo.ToString(), EntityIndex + 1);
+							return FString::Printf(TEXT("%s.%s%i"), *ExposedProperty->FieldPathInfo.ToString(), *GetSubobjectPath(), EntityIndex + 1);
 						}
 					}();
 
@@ -117,8 +117,6 @@ namespace UE::RemoteControl::DMX
 			if (DMXEntity->ExtraSetting.FixturePatchReference.GetFixturePatch() != UnifiedFixturePatch ||
 				DMXEntity->ExtraSetting.bIsPrimaryPatch != bUnifiedIsPrimaryPatch)
 			{
-				ensureMsgf(0, TEXT("Found entities of the same property with differing properties. Adopting property from first entity to mend."));
-
 				DMXEntity->ExtraSetting.FixturePatchReference = UnifiedFixturePatch;
 				DMXEntity->ExtraSetting.bIsPrimaryPatch = bUnifiedIsPrimaryPatch;
 			}
diff --git a/Engine/Plugins/VirtualProduction/RemoteControlProtocolDMX/Source/RemoteControlProtocolDMX/Private/Library/RemoteControlDMXControlledPropertyPatch.cpp b/Engine/Plugins/VirtualProduction/RemoteControlProtocolDMX/Source/RemoteControlProtocolDMX/Private/Library/RemoteControlDMXControlledPropertyPatch.cpp
--- a/Engine/Plugins/VirtualProduction/RemoteControlProtocolDMX/Source/RemoteControlProtocolDMX/Private/Library/RemoteControlDMXControlledPropertyPatch.cpp
+++ b/Engine/Plugins/VirtualProduction/RemoteControlProtocolDMX/Source/RemoteControlProtocolDMX/Private/Library/RemoteControlDMXControlledPropertyPatch.cpp
@@ -35,6 +35,11 @@ namespace UE::RemoteControl::DMX
 		return DMXControlledProperties.IsEmpty() ? nullptr : DMXControlledProperties[0]->GetFixturePatch();
 	}
 
+	const UObject* FRemoteControlDMXControlledPropertyPatch::GetOwnerObject() const
+	{
+		return DMXControlledProperties.IsEmpty() ? nullptr : DMXControlledProperties[0]->GetOwnerObject();
+	}
+
 #if WITH_EDITOR
 	void FRemoteControlDMXControlledPropertyPatch::OnFixtureTypeChanged(const UDMXEntityFixtureType* ChangedFixtureType)
 	{
diff --git a/Engine/Plugins/VirtualProduction/RemoteControlProtocolDMX/Source/RemoteControlProtocolDMX/Private/Library/RemoteControlDMXLibraryProxy.cpp b/Engine/Plugins/VirtualProduction/RemoteControlProtocolDMX/Source/RemoteControlProtocolDMX/Private/Library/RemoteControlDMXLibraryProxy.cpp
--- a/Engine/Plugins/VirtualProduction/RemoteControlProtocolDMX/Source/RemoteControlProtocolDMX/Private/Library/RemoteControlDMXLibraryProxy.cpp
+++ b/Engine/Plugins/VirtualProduction/RemoteControlProtocolDMX/Source/RemoteControlProtocolDMX/Private/Library/RemoteControlDMXLibraryProxy.cpp
@@ -8,12 +8,14 @@
 #include "Library/DMXLibrary.h"
 #include "Library/RemoteControlDMXControlledProperty.h"
 #include "Library/RemoteControlDMXControlledPropertyPatch.h"
+#include "Library/RemoteControlDMXProtocolEntityObserver.h"
 #include "Misc/CoreDelegates.h"
 #include "RemoteControlDMXLog.h"
 #include "RemoteControlDMXUserData.h"
 #include "RemoteControlField.h"
 #include "RemoteControlPreset.h"
 #include "RemoteControlProtocolDMX.h"
+#include "UObject/UObjectGlobals.h"
 
 #if WITH_EDITOR
 FRemoteControlDMXPrePropertyPatchesChanged URemoteControlDMXLibraryProxy::OnPrePropertyPatchesChanged;
@@ -40,6 +42,8 @@ void URemoteControlDMXLibraryProxy::PostInitProperties()
 	Preset->OnEntityUnexposed().AddUObject(this, &URemoteControlDMXLibraryProxy::OnEntityExposedOrUnexposed);
 	Preset->OnEntityRebind().AddUObject(this, &URemoteControlDMXLibraryProxy::OnEntityRebind);
 	Preset->OnEntitiesUpdated().AddUObject(this, &URemoteControlDMXLibraryProxy::OnEntitiesUpdated);
+
+	FCoreUObjectDelegates::PostLoadMapWithWorld.AddUObject(this, &URemoteControlDMXLibraryProxy::OnPostLoadMapWithWorld);
 }
 
 void URemoteControlDMXLibraryProxy::PostLoad()
@@ -103,6 +107,9 @@ void URemoteControlDMXLibraryProxy::Refresh()
 		
 #if WITH_EDITOR
 		OnPostPropertyPatchesChanged.Broadcast();
+		
+		// Listen to DMX related property changes of entities in editor
+		UpdateEntitiesObserver();
 #endif 
 
 		BindOnFixturePatchesReceived();
@@ -111,9 +118,11 @@ void URemoteControlDMXLibraryProxy::Refresh()
 
 void URemoteControlDMXLibraryProxy::Reset()
 {
+	// Cancel any refresh requests
 	RefreshDelegateHandle.Reset();
 	FCoreDelegates::OnEndFrame.RemoveAll(this);
 
+	// Reset property patches
 	PropertyPatches.Reset();
 }
 
@@ -240,6 +249,22 @@ void URemoteControlDMXLibraryProxy::UnbindOnFixturePatchesReceived()
 	}
 }
 
+#if WITH_EDITOR
+void URemoteControlDMXLibraryProxy::UpdateEntitiesObserver()
+{
+	TArray<TSharedRef<TStructOnScope<FRemoteControlProtocolEntity>>> Entities;
+	for (const TSharedRef<FRemoteControlDMXControlledPropertyPatch>& PropertyPatch : PropertyPatches)
+	{
+		for (const TSharedRef<FRemoteControlDMXControlledProperty>& Property : PropertyPatch->GetDMXControlledProperties())
+		{
+			Entities.Append(Property->GetEntities());
+		}
+	}
+
+	EntitiesObserver = MakeShared<FRemoteControlDMXProtocolEntityObserver>(Entities);
+}
+#endif // WITH_EDITOR
+
 void URemoteControlDMXLibraryProxy::OnFixturePatchReceived(UDMXEntityFixturePatch* FixturePatch, const FDMXNormalizedAttributeValueMap& ValuePerAttribute)
 {
 	const TSharedPtr<IRemoteControlProtocol> DMXProtocol = IRemoteControlProtocolModule::Get().GetProtocolByName(FRemoteControlProtocolDMX::ProtocolName);
@@ -349,3 +374,8 @@ void URemoteControlDMXLibraryProxy::OnEntitiesUpdated(URemoteControlPreset* Pres
 {
 	RequestRefresh();
 }
+
+void URemoteControlDMXLibraryProxy::OnPostLoadMapWithWorld(UWorld* World)
+{
+	RequestRefresh();
+}
diff --git a/Engine/Plugins/VirtualProduction/RemoteControlProtocolDMX/Source/RemoteControlProtocolDMX/Private/Library/RemoteControlDMXProtocolEntityObserver.cpp b/Engine/Plugins/VirtualProduction/RemoteControlProtocolDMX/Source/RemoteControlProtocolDMX/Private/Library/RemoteControlDMXProtocolEntityObserver.cpp
new file mode 100644
--- /dev/null
+++ b/Engine/Plugins/VirtualProduction/RemoteControlProtocolDMX/Source/RemoteControlProtocolDMX/Private/Library/RemoteControlDMXProtocolEntityObserver.cpp
@@ -0,0 +1,92 @@
+// Copyright Epic Games, Inc. All Rights Reserved.
+
+#if WITH_EDITOR
+
+#include "RemoteControlDMXProtocolEntityObserver.h"
+
+#include "Library/RemoteControlDMXLibraryProxy.h"
+#include "RemoteControlDMXUserData.h"
+#include "RemoteControlPreset.h"
+#include "RemoteControlProtocolDMX.h"
+
+namespace UE::RemoteControl::DMX
+{
+	FRemoteControlDMXProtocolEntityComparator::FRemoteControlDMXProtocolEntityComparator(TSharedRef<TStructOnScope<FRemoteControlProtocolEntity>> InEntity)
+		: Entity(InEntity)
+	{
+		UpdateCache();
+	}
+
+	bool FRemoteControlDMXProtocolEntityComparator::DidPropertiesChange() const
+	{
+		const FRemoteControlDMXProtocolEntity* DMXEntity = Entity->IsValid() ? Entity->Cast<FRemoteControlDMXProtocolEntity>() : nullptr;
+		if (DMXEntity)
+		{
+			const bool bChanged = 
+				DMXEntity->ExtraSetting.bUseLSB != bUseLSBMode ||
+				DMXEntity->ExtraSetting.DataType != DataType;
+
+			if (bChanged)
+			{
+				UpdateCache();
+				return true;
+			}
+		}
+	
+		return false;
+	}
+
+	void FRemoteControlDMXProtocolEntityComparator::UpdateCache() const
+	{
+		const FRemoteControlDMXProtocolEntity* DMXEntity = Entity->IsValid() ? Entity->Cast<FRemoteControlDMXProtocolEntity>() : nullptr;
+		if (DMXEntity)
+		{
+			bUseLSBMode = DMXEntity->ExtraSetting.bUseLSB;
+			DataType = DMXEntity->ExtraSetting.DataType;
+		}
+	}
+
+	FRemoteControlDMXProtocolEntityObserver::FRemoteControlDMXProtocolEntityObserver(const TArray<TSharedRef<TStructOnScope<FRemoteControlProtocolEntity>>>& Entities)
+	{
+		Algo::Transform(Entities, Comparators,
+			[](const TSharedRef<TStructOnScope<FRemoteControlProtocolEntity>>& Entity)
+			{
+				return FRemoteControlDMXProtocolEntityComparator(Entity);
+			});
+	}
+
+	void FRemoteControlDMXProtocolEntityObserver::Tick(float DeltaTime)
+	{
+		for (const FRemoteControlDMXProtocolEntityComparator& Comparator : Comparators)
+		{
+			if (Comparator.DidPropertiesChange())
+			{
+				if (URemoteControlDMXLibraryProxy* DMXLibraryProxy = GetDMXLibraryProxy(Comparator.Entity))
+				{
+					DMXLibraryProxy->RequestRefresh();
+					return;
+				}
+			}
+		}
+	}
+
+	TStatId FRemoteControlDMXProtocolEntityObserver::GetStatId() const
+	{
+		RETURN_QUICK_DECLARE_CYCLE_STAT(FRemoteControlDMXProtocolEntityObserver, STATGROUP_Tickables);
+	}
+
+	URemoteControlDMXLibraryProxy* FRemoteControlDMXProtocolEntityObserver::GetDMXLibraryProxy(const TSharedRef<TStructOnScope<FRemoteControlProtocolEntity>>& Entity) const
+	{
+		const FRemoteControlProtocolEntity* EntityPtr = Entity->IsValid() ? Entity->Get() : nullptr;
+		URemoteControlPreset* Preset = EntityPtr ? EntityPtr->GetOwner().Get() : nullptr;;
+		if (Preset)
+		{
+			URemoteControlDMXUserData* DMXUserData = URemoteControlDMXUserData::GetOrCreateDMXUserData(Preset);
+			return DMXUserData->GetDMXLibraryProxy();
+		}
+
+		return nullptr;
+	}
+}
+
+#endif
diff --git a/Engine/Plugins/VirtualProduction/RemoteControlProtocolDMX/Source/RemoteControlProtocolDMX/Private/Library/RemoteControlDMXProtocolEntityObserver.h b/Engine/Plugins/VirtualProduction/RemoteControlProtocolDMX/Source/RemoteControlProtocolDMX/Private/Library/RemoteControlDMXProtocolEntityObserver.h
new file mode 100644
--- /dev/null
+++ b/Engine/Plugins/VirtualProduction/RemoteControlProtocolDMX/Source/RemoteControlProtocolDMX/Private/Library/RemoteControlDMXProtocolEntityObserver.h
@@ -0,0 +1,67 @@
+// Copyright Epic Games, Inc. All Rights Reserved.
+
+#pragma once
+
+#if WITH_EDITOR
+
+#include "DMXProtocolTypes.h"
+#include "Templates/SharedPointer.h"
+#include "Tickable.h"
+#include "UObject/StructOnScope.h"
+
+class URemoteControlDMXLibraryProxy;
+struct FRemoteControlProtocolEntity;
+
+namespace UE::RemoteControl::DMX
+{
+	struct FRemoteControlDMXProtocolEntityComparator
+	{
+		FRemoteControlDMXProtocolEntityComparator(TSharedRef<TStructOnScope<FRemoteControlProtocolEntity>> InEntity);
+
+		/** Returns true if properties changed */
+		bool DidPropertiesChange() const;
+
+		/** The entity to compare */
+		const TSharedRef<TStructOnScope<FRemoteControlProtocolEntity>> Entity;
+
+	private:
+		/** Updates cached values */
+		void UpdateCache() const;
+
+		/** Cached LSB Mode */
+		mutable bool bUseLSBMode = false;
+
+		/** Cached Signal Format */
+		mutable EDMXFixtureSignalFormat DataType = EDMXFixtureSignalFormat::E8Bit;
+	};
+
+	/** 
+	 * Class that observes and notifies about DMX protocol entity changes. 
+	 * 
+	 * Required since the related RC event URemoteControlPreset::OnExposedPropertiesModified fires on RC changed properties,
+	 * however we only want to handle DMX specific property changes. 
+	 */
+	class FRemoteControlDMXProtocolEntityObserver
+		: public FTickableGameObject
+		, public TSharedFromThis<FRemoteControlDMXProtocolEntityObserver>
+	{
+	public:
+		FRemoteControlDMXProtocolEntityObserver() = default;
+		FRemoteControlDMXProtocolEntityObserver(const TArray<TSharedRef<TStructOnScope<FRemoteControlProtocolEntity>>>& Entities);
+
+	private:
+		//~ Begin FTickableGameObject interface
+		virtual bool IsTickableInEditor() const override { return true; }
+		virtual void Tick(float DeltaTime) override;
+		virtual TStatId GetStatId() const override;
+		//~ End FTickableGameObject interface
+		
+		/** Returns the DMX Library Proxy for the specified Entity */
+		URemoteControlDMXLibraryProxy* GetDMXLibraryProxy(const TSharedRef<TStructOnScope<FRemoteControlProtocolEntity>>& Entity) const;
+
+		/** Per entity comparator to detect changes */
+		TArray<FRemoteControlDMXProtocolEntityComparator> Comparators;
+	};
+}
+
+#endif
diff --git a/Engine/Plugins/VirtualProduction/RemoteControlProtocolDMX/Source/RemoteControlProtocolDMXEditor/Private/Library/RemoteControlDMXLibraryBuilder.cpp b/Engine/Plugins/VirtualProduction/RemoteControlProtocolDMX/Source/RemoteControlProtocolDMXEditor/Private/Library/RemoteControlDMXLibraryBuilder.cpp
--- a/Engine/Plugins/VirtualProduction/RemoteControlProtocolDMX/Source/RemoteControlProtocolDMXEditor/Private/Library/RemoteControlDMXLibraryBuilder.cpp
+++ b/Engine/Plugins/VirtualProduction/RemoteControlProtocolDMX/Source/RemoteControlProtocolDMXEditor/Private/Library/RemoteControlDMXLibraryBuilder.cpp
@@ -3,6 +3,8 @@
 #include "RemoteControlDMXLibraryBuilder.h"
 
 #include "Algo/AnyOf.h"
+#include "Algo/MaxElement.h"
+#include "Editor.h"
 #include "RemoteControlDMXUserData.h"
 #include "RemoteControlPreset.h"
 #include "Library/DMXEntityFixturePatch.h"
@@ -13,6 +15,8 @@
 
 namespace UE::RemoteControl::DMX
 {
+	const FString FRemoteControlDMXLibraryBuilder::RCFixtureGroupTag = TEXT("RCGenerated_PatchGroup: ");
+
 	void FRemoteControlDMXLibraryBuilder::Register()
 	{
 		static TSharedRef<FRemoteControlDMXLibraryBuilder> Instance = MakeShared<FRemoteControlDMXLibraryBuilder>();
@@ -37,7 +41,9 @@ namespace UE::RemoteControl::DMX
 			Algo::TransformIf(PreEditChangePropertyPatches, PreviousFixturePatches,
 				[](const TSharedRef<FRemoteControlDMXControlledPropertyPatch>& PropertyPatch)
 				{
-					return PropertyPatch->GetFixturePatch();
+					return
+						PropertyPatch->GetOwnerObject() &&
+						PropertyPatch->GetFixturePatch();
 				},
 				[](const TSharedRef<FRemoteControlDMXControlledPropertyPatch>& PropertyPatch)
 				{
@@ -163,54 +169,166 @@ namespace UE::RemoteControl::DMX
 				return PropertyPatch->GetFixturePatch();
 			});
 
-		// Reset patch
+		// To retain auto-assign order for patches for properties in different worlds, use a group index.
+		const int32 GroupIndex = GetOrCreateGroupIndex(*DMXLibrary, FixturePatches);
+		const FName Tag = *(RCFixtureGroupTag + FString::FromInt(GroupIndex));
 		for (UDMXEntityFixturePatch* FixturePatch : FixturePatches)
 		{
-			FixturePatch->SetStartingChannel(1);
-			FixturePatch->SetUniverseID(1);
+			FixturePatch->CustomTags.AddUnique(Tag);
+		}
+
+		// Acquire all RC related patches
+		const TMap<int32, TArray<UDMXEntityFixturePatch*>> GroupIndexToRCFixturePatchesMap = GetGroupIndexToRCFixturePatchesMap(*DMXLibrary);
+
+		// Reset patches
+		for (const TTuple<int32, TArray<UDMXEntityFixturePatch*>>& GroupIndexToRCFixturePatchesPair : GroupIndexToRCFixturePatchesMap)
+		{
+			for (UDMXEntityFixturePatch* FixturePatch : GroupIndexToRCFixturePatchesPair.Value)
+			{
+				FixturePatch->PreEditChange(nullptr);
+
+				FixturePatch->SetStartingChannel(1);
+				FixturePatch->SetUniverseID(1);
+			}
+		}
+
+		// Auto assign all RC related patches
+		UDMXEntityFixturePatch* PreviousFixturePatch = nullptr;
+		for (const TTuple<int32, TArray<UDMXEntityFixturePatch*>>& GroupIndexToRCFixturePatchesPair : GroupIndexToRCFixturePatchesMap)
+		{
+			for (UDMXEntityFixturePatch* FixturePatch : GroupIndexToRCFixturePatchesPair.Value)
+			{
+				const int32 AutoAssingFromUniverse = DMXUserData->GetAutoAssignFromUniverse();
+				const int32	DesiredAbsoluteStartingChannel = [PreviousFixturePatch, AutoAssingFromUniverse]()
+					{
+						if (PreviousFixturePatch && PreviousFixturePatch->GetUniverseID() >= AutoAssingFromUniverse)
+						{
+							return
+								(int64)PreviousFixturePatch->GetUniverseID() * DMX_UNIVERSE_SIZE +
+								PreviousFixturePatch->GetEndingChannel();
+						}
+						else
+						{
+							return (int64)AutoAssingFromUniverse * DMX_UNIVERSE_SIZE;
+						}
+					}();
+
+				const bool bFitsUniverse = DesiredAbsoluteStartingChannel % DMX_UNIVERSE_SIZE + FixturePatch->GetChannelSpan() <= DMX_UNIVERSE_SIZE;
+				const int64 AbsoluteStartingChannel = bFitsUniverse ? DesiredAbsoluteStartingChannel : (DesiredAbsoluteStartingChannel / DMX_UNIVERSE_SIZE + 1) * DMX_UNIVERSE_SIZE;
+
+				const int32 Universe = AbsoluteStartingChannel / DMX_UNIVERSE_SIZE;
+				const int32 Channel = AbsoluteStartingChannel % DMX_UNIVERSE_SIZE + 1;
+
+				FixturePatch->SetUniverseID(Universe);
+				FixturePatch->SetStartingChannel(Channel);
+
+				FixturePatch->PostEditChange();
+
+				PreviousFixturePatch = FixturePatch;
+			}
 		}
+	}
+
+	TMap<int32, TArray<UDMXEntityFixturePatch*>> FRemoteControlDMXLibraryBuilder::GetGroupIndexToRCFixturePatchesMap(const UDMXLibrary& DMXLibrary) const
+	{
+		TMap<int32, TArray<UDMXEntityFixturePatch*>> Result;
 
-		// Auto assign
-		for (const TSharedRef<FRemoteControlDMXControlledPropertyPatch>& PropertyPatch : PostEditChangePropertyPatches)
+		const TArray<UDMXEntityFixturePatch*> AllFixturePatchesInLibrary = DMXLibrary.GetEntitiesTypeCast<UDMXEntityFixturePatch>();
+		for (UDMXEntityFixturePatch* FixturePatch : AllFixturePatchesInLibrary)
 		{
-			UDMXEntityFixturePatch* FixturePatch = PropertyPatch->GetFixturePatch();
 			if (!FixturePatch)
 			{
 				continue;
 			}
 
-			const int32 FixturePatchIndex = FixturePatches.IndexOfByKey(FixturePatch);
-			if (!ensureMsgf(FixturePatchIndex != INDEX_NONE, TEXT("Unexpected cannot find fixture patch in DMX Library. Cannot auto assign fixture patch.")))
+			const FName* GroupTagPtr = Algo::FindByPredicate(FixturePatch->CustomTags, [](const FName& Tag)
+				{
+					return Tag.ToString().StartsWith(RCFixtureGroupTag);
+				});
+
+			const int32 GroupIndex = GroupTagPtr ? ExtractGroupIndex(*GroupTagPtr) : INDEX_NONE;
+			if (GroupIndex != INDEX_NONE)
 			{
-				continue;
+				Result.FindOrAdd(GroupIndex).Add(FixturePatch);
 			}
+		}
 
-			const int32 AutoAssingFromUniverse = DMXUserData->GetAutoAssignFromUniverse();
-			const UDMXEntityFixturePatch* PreviousPatch = FixturePatches.IsValidIndex(FixturePatchIndex - 1) ? FixturePatches[FixturePatchIndex - 1] : nullptr;
-			const int64 DesiredAbsoluteStartingChannel = [PreviousPatch, AutoAssingFromUniverse]()
+		Result.KeyStableSort([](const int32 GroupIndexA, const int32 GroupIndexB)
+			{
+				return GroupIndexA <= GroupIndexB;
+			});
+
+		return Result;
+	}
+
+	int32 FRemoteControlDMXLibraryBuilder::GetOrCreateGroupIndex(const UDMXLibrary& DMXLibrary, const TArray<UDMXEntityFixturePatch*>& FixturePatches) const
+	{
+		const FName* PreviouslySetTagPtr = [&FixturePatches]() -> const FName*
+			{
+				for (UDMXEntityFixturePatch* FixturePatch : FixturePatches)
 				{
-					if (PreviousPatch && PreviousPatch->GetUniverseID() >= AutoAssingFromUniverse)
-					{
-						return (int64)PreviousPatch->GetUniverseID() * DMX_UNIVERSE_SIZE + PreviousPatch->GetEndingChannel();
-					}
-					else
+					const FName* GroupTagPtr = FixturePatch ?
+						Algo::FindByPredicate(FixturePatch->CustomTags, [](const FName& Tag)
+							{
+								return Tag.ToString().StartsWith(RCFixtureGroupTag);
+							}) :
+						nullptr;
+
+					if (GroupTagPtr)
 					{
-						return (int64)AutoAssingFromUniverse * DMX_UNIVERSE_SIZE;
+						return GroupTagPtr;
 					}
-				}();
+				}
+				return nullptr;
+			}();
 
-			const bool bFitsUniverse = DesiredAbsoluteStartingChannel % DMX_UNIVERSE_SIZE + FixturePatch->GetChannelSpan() <= DMX_UNIVERSE_SIZE;
-			const int64 AbsoluteStartingChannel = bFitsUniverse ? DesiredAbsoluteStartingChannel : (DesiredAbsoluteStartingChannel / DMX_UNIVERSE_SIZE + 1) * DMX_UNIVERSE_SIZE;
 
-			const int32 Universe = AbsoluteStartingChannel / DMX_UNIVERSE_SIZE;
-			const int32 Channel = AbsoluteStartingChannel % DMX_UNIVERSE_SIZE + 1;
+		const int32 GroupIndex = PreviouslySetTagPtr ? ExtractGroupIndex(*PreviouslySetTagPtr) : INDEX_NONE;
+		if (GroupIndex != INDEX_NONE)
+		{
+			return GroupIndex;
+		}
+		else
+		{
+			// Return the next free group index
+			const TArray<UDMXEntityFixturePatch*> AllFixturePatchesInLibrary = DMXLibrary.GetEntitiesTypeCast<UDMXEntityFixturePatch>();
+			int32 NextGroupIndex = 0;
 
-			FixturePatch->PreEditChange(nullptr);
+			for (const UDMXEntityFixturePatch* FixturePatch : AllFixturePatchesInLibrary)
+			{
+				if (!FixturePatch || FixturePatches.Contains(FixturePatch))
+				{
+					continue;
+				}
+
+				const FName* OtherRCTagPtr = Algo::FindByPredicate(FixturePatch->CustomTags, [](const FName& Tag)
+					{
+						return Tag.ToString().StartsWith(RCFixtureGroupTag);
+					});
 
-			FixturePatch->SetUniverseID(Universe);
-			FixturePatch->SetStartingChannel(Channel);
+				const int32 OtherGroupIndex = OtherRCTagPtr ? ExtractGroupIndex(*OtherRCTagPtr) : INDEX_NONE;
+				if (OtherGroupIndex != INDEX_NONE)
+				{
+					NextGroupIndex = FMath::Max(NextGroupIndex, OtherGroupIndex + 1);
+				}
+			}
 
-			FixturePatch->PostEditChange();
+			return NextGroupIndex;
 		}
 	}
+
+	int32 FRemoteControlDMXLibraryBuilder::ExtractGroupIndex(const FName& Tag) const
+	{
+		FString TagString = Tag.ToString();
+		if (TagString.RemoveFromStart(RCFixtureGroupTag))
+		{
+			int32 Result;
+			if (LexTryParseString(Result, *TagString))
+			{
+				return Result;
+			}
+		}
+
+		return INDEX_NONE;
+	}
 }
diff --git a/Engine/Plugins/VirtualProduction/RemoteControlProtocolDMX/Source/RemoteControlProtocolDMXEditor/Private/Library/RemoteControlDMXLibraryBuilder.h b/Engine/Plugins/VirtualProduction/RemoteControlProtocolDMX/Source/RemoteControlProtocolDMXEditor/Private/Library/RemoteControlDMXLibraryBuilder.h
--- a/Engine/Plugins/VirtualProduction/RemoteControlProtocolDMX/Source/RemoteControlProtocolDMXEditor/Private/Library/RemoteControlDMXLibraryBuilder.h
+++ b/Engine/Plugins/VirtualProduction/RemoteControlProtocolDMX/Source/RemoteControlProtocolDMXEditor/Private/Library/RemoteControlDMXLibraryBuilder.h
@@ -42,6 +42,15 @@ namespace UE::RemoteControl::DMX
 		/** Auto assigns fixture patches in the DMX Library */
 		void AutoAssignFixturePatches(const TArray<TSharedRef<FRemoteControlDMXControlledPropertyPatch>>& PostEditChangePropertyPatches);
 
+		/** Returns all Remote Control related Fixture Patches in the DMX Library, grouped by group index */
+		TMap<int32, TArray<UDMXEntityFixturePatch*>> GetGroupIndexToRCFixturePatchesMap(const UDMXLibrary& DMXLibrary) const;
+		
+		/** Gets or creates a Group Index for a group of Fixture Patches. */
+		int32 GetOrCreateGroupIndex(const UDMXLibrary& DMXLibrary, const TArray<UDMXEntityFixturePatch*>& FixturePatches) const;
+
+		/** Extracts a group index from a Fixture Patch tag */
+		int32 ExtractGroupIndex(const FName& Tag) const;
+
 		/** Fixture types used before property patches changed */
 		TArray<UDMXEntityFixtureType*> PreviousFixtureTypes;
 
@@ -50,5 +59,8 @@ namespace UE::RemoteControl::DMX
 
 		/** The remote control preset for which the Library is built */
 		URemoteControlPreset* Preset = nullptr;
+
+		/** Tag for RC Fixture Groups */
+		static const FString RCFixtureGroupTag;
 	};
 }
diff --git a/Engine/Plugins/VirtualProduction/RemoteControlProtocolDMX/Source/RemoteControlProtocolDMXEditor/Private/Library/RemoteControlDMXPatchBuilder.cpp b/Engine/Plugins/VirtualProduction/RemoteControlProtocolDMX/Source/RemoteControlProtocolDMXEditor/Private/Library/RemoteControlDMXPatchBuilder.cpp
--- a/Engine/Plugins/VirtualProduction/RemoteControlProtocolDMX/Source/RemoteControlProtocolDMXEditor/Private/Library/RemoteControlDMXPatchBuilder.cpp
+++ b/Engine/Plugins/VirtualProduction/RemoteControlProtocolDMX/Source/RemoteControlProtocolDMXEditor/Private/Library/RemoteControlDMXPatchBuilder.cpp
@@ -289,16 +289,20 @@ namespace UE::RemoteControl::DMX
 						const FDMXFixtureFunction& Function = Mode.Functions[FunctionIndex];
 						const FRemoteControlDMXProtocolEntity* DMXEntity = DMXEntities.IsValidIndex(FunctionIndex) ? DMXEntities[FunctionIndex] : nullptr;
 
-						if (DMXEntity &&
-							DMXEntity->ExtraSetting.AttributeName == Function.Attribute.Name &&
-							DMXEntity->ExtraSetting.DataType == Function.DataType &&
-							DMXEntity->ExtraSetting.bUseLSB == Function.bUseLSBMode)
+						if (!DMXEntity)
 						{
-							return true;
+							continue;
+						}
+
+						if (DMXEntity->ExtraSetting.AttributeName != Function.Attribute.Name ||
+							DMXEntity->ExtraSetting.DataType != Function.DataType ||
+							DMXEntity->ExtraSetting.bUseLSB != Function.bUseLSBMode)
+						{
+							return false;
 						}
 					}
 
-					return false;
+					return true;
 				});
 
 			return FixtureTypePtr ? *FixtureTypePtr : nullptr;
diff --git a/Engine/Plugins/VirtualProduction/RemoteControlProtocolDMX/Source/RemoteControlProtocolDMXEditor/Private/RemoteControlDMXPerPresetEditorSettings.cpp b/Engine/Plugins/VirtualProduction/RemoteControlProtocolDMX/Source/RemoteControlProtocolDMXEditor/Private/RemoteControlDMXPerPresetEditorSettings.cpp
--- a/Engine/Plugins/VirtualProduction/RemoteControlProtocolDMX/Source/RemoteControlProtocolDMXEditor/Private/RemoteControlDMXPerPresetEditorSettings.cpp
+++ b/Engine/Plugins/VirtualProduction/RemoteControlProtocolDMX/Source/RemoteControlProtocolDMXEditor/Private/RemoteControlDMXPerPresetEditorSettings.cpp
@@ -3,6 +3,7 @@
 #include "RemoteControlDMXPerPresetEditorSettings.h"
 
 #include "RemoteControlDMXUserData.h"
+#include "UI/RCPanelExposedEntitiesListSettingsData.h"
 
 URemoteControlDMXPerPresetEditorSettings* URemoteControlDMXPerPresetEditorSettings::GetOrCreatePerPresetEditorSettings(URemoteControlPreset* Preset)
 {
@@ -14,7 +15,15 @@ URemoteControlDMXPerPresetEditorSettings* URemoteControlDMXPerPresetEditorSettin
 	}
 	else
 	{
-		DMXUserData->PerPresetEditorSettings = NewObject<URemoteControlDMXPerPresetEditorSettings>(DMXUserData, NAME_None);
+		URemoteControlDMXPerPresetEditorSettings* NewSettings = NewObject<URemoteControlDMXPerPresetEditorSettings>(DMXUserData, NAME_None);
+		
+		// Set default settings for Remote Control DMX.
+		NewSettings->ExposedEntitiesListSettings.FieldGroupType = ERCFieldGroupType::Owner;
+		NewSettings->ExposedEntitiesListSettings.FieldGroupOrder = ERCFieldGroupOrder::Ascending;
+
+		// Store it in DMX User Data
+		DMXUserData->PerPresetEditorSettings = NewSettings;
+
 		return CastChecked<URemoteControlDMXPerPresetEditorSettings>(DMXUserData->PerPresetEditorSettings);
 	}
 }
diff --git a/Engine/Plugins/VirtualProduction/RemoteControlProtocolDMX/Source/RemoteControlProtocolDMXEditor/Private/UI/SRemoteControlDMXExposedEntitiesGroupPatch.cpp b/Engine/Plugins/VirtualProduction/RemoteControlProtocolDMX/Source/RemoteControlProtocolDMXEditor/Private/UI/SRemoteControlDMXExposedEntitiesGroupPatch.cpp
--- a/Engine/Plugins/VirtualProduction/RemoteControlProtocolDMX/Source/RemoteControlProtocolDMXEditor/Private/UI/SRemoteControlDMXExposedEntitiesGroupPatch.cpp
+++ b/Engine/Plugins/VirtualProduction/RemoteControlProtocolDMX/Source/RemoteControlProtocolDMXEditor/Private/UI/SRemoteControlDMXExposedEntitiesGroupPatch.cpp
@@ -67,6 +67,18 @@ namespace UE::RemoteControl::DMX
 	{
 		RefreshTimerHandle.Invalidate();
 
+		// Only show the patch when grouped by owner
+		const URemoteControlDMXUserData* DMXUserData = GetDMXUserData();
+		if (!DMXUserData || DMXUserData->GetPatchGroupMode() != ERemoteControlDMXPatchGroupMode::GroupByOwner)
+		{
+			ChildSlot
+				[
+					SNullWidget::NullWidget
+				];
+
+			return;
+		}
+
 		// Update the fixture patches that can be used as primary for this patch
 		WeakFixturePatches.Reset();
 
@@ -94,7 +106,6 @@ namespace UE::RemoteControl::DMX
 			}
 		}
 
-
 		if (Entities.IsEmpty())
 		{
 			ChildSlot
