Remote Control DMX: Transform of controlled objects are now written to the DMX Library

diff --git a/Engine/Plugins/VirtualProduction/RemoteControlProtocolDMX/Source/RemoteControlProtocolDMX/Internal/Library/RemoteControlDMXControlledProperty.h b/Engine/Plugins/VirtualProduction/RemoteControlProtocolDMX/Source/RemoteControlProtocolDMX/Internal/Library/RemoteControlDMXControlledProperty.h
--- a/Engine/Plugins/VirtualProduction/RemoteControlProtocolDMX/Source/RemoteControlProtocolDMX/Internal/Library/RemoteControlDMXControlledProperty.h
+++ b/Engine/Plugins/VirtualProduction/RemoteControlProtocolDMX/Source/RemoteControlProtocolDMX/Internal/Library/RemoteControlDMXControlledProperty.h
@@ -23,7 +23,7 @@ namespace UE::RemoteControl::DMX
 		TArray<TSharedRef<TStructOnScope<FRemoteControlProtocolEntity>>> GetEntities() const { return Entities; }
 
 		/** Returns the owner object of this property */
-		const UObject* GetOwnerObject() const;
+		const UObject* GetOwnerActor() const;
 
 		/** Returns the fixture patch of this property */
 		UDMXEntityFixturePatch* GetFixturePatch() const;
diff --git a/Engine/Plugins/VirtualProduction/RemoteControlProtocolDMX/Source/RemoteControlProtocolDMX/Internal/Library/RemoteControlDMXControlledPropertyPatch.h b/Engine/Plugins/VirtualProduction/RemoteControlProtocolDMX/Source/RemoteControlProtocolDMX/Internal/Library/RemoteControlDMXControlledPropertyPatch.h
--- a/Engine/Plugins/VirtualProduction/RemoteControlProtocolDMX/Source/RemoteControlProtocolDMX/Internal/Library/RemoteControlDMXControlledPropertyPatch.h
+++ b/Engine/Plugins/VirtualProduction/RemoteControlProtocolDMX/Source/RemoteControlProtocolDMX/Internal/Library/RemoteControlDMXControlledPropertyPatch.h
@@ -35,8 +35,8 @@ namespace UE::RemoteControl::DMX
 		/** Returns the fixture patch corresponding to this property patch */
 		UDMXEntityFixturePatch* GetFixturePatch() const;
 
-		/** Returns the owner object of this patch */
-		const UObject* GetOwnerObject() const;
+		/** Returns the owner actor of this patch, or nullptr if there is no outer Actor. */
+		const UObject* GetOwnerActor() const;
 
 	private:
 #if WITH_EDITOR
diff --git a/Engine/Plugins/VirtualProduction/RemoteControlProtocolDMX/Source/RemoteControlProtocolDMX/Private/Library/RemoteControlDMXControlledProperty.cpp b/Engine/Plugins/VirtualProduction/RemoteControlProtocolDMX/Source/RemoteControlProtocolDMX/Private/Library/RemoteControlDMXControlledProperty.cpp
--- a/Engine/Plugins/VirtualProduction/RemoteControlProtocolDMX/Source/RemoteControlProtocolDMX/Private/Library/RemoteControlDMXControlledProperty.cpp
+++ b/Engine/Plugins/VirtualProduction/RemoteControlProtocolDMX/Source/RemoteControlProtocolDMX/Private/Library/RemoteControlDMXControlledProperty.cpp
@@ -28,7 +28,7 @@ namespace UE::RemoteControl::DMX
 		UnifyEntities();
 	}
 
-	const UObject* FRemoteControlDMXControlledProperty::GetOwnerObject() const
+	const UObject* FRemoteControlDMXControlledProperty::GetOwnerActor() const
 	{
 		const UObject* BoundObject = ExposedProperty->GetBoundObject();
 		if (!BoundObject)
@@ -125,7 +125,7 @@ namespace UE::RemoteControl::DMX
 
 	FString FRemoteControlDMXControlledProperty::GetSubobjectPath() const
 	{
-		if (const UObject* OwnerObject = GetOwnerObject())
+		if (const UObject* OwnerObject = GetOwnerActor())
 		{
 			const FString OwnerName = OwnerObject->GetFName().ToString();
 			const FString BindingPath = ExposedProperty->GetLastBindingPath().ToString();
diff --git a/Engine/Plugins/VirtualProduction/RemoteControlProtocolDMX/Source/RemoteControlProtocolDMX/Private/Library/RemoteControlDMXControlledPropertyPatch.cpp b/Engine/Plugins/VirtualProduction/RemoteControlProtocolDMX/Source/RemoteControlProtocolDMX/Private/Library/RemoteControlDMXControlledPropertyPatch.cpp
--- a/Engine/Plugins/VirtualProduction/RemoteControlProtocolDMX/Source/RemoteControlProtocolDMX/Private/Library/RemoteControlDMXControlledPropertyPatch.cpp
+++ b/Engine/Plugins/VirtualProduction/RemoteControlProtocolDMX/Source/RemoteControlProtocolDMX/Private/Library/RemoteControlDMXControlledPropertyPatch.cpp
@@ -35,9 +35,9 @@ namespace UE::RemoteControl::DMX
 		return DMXControlledProperties.IsEmpty() ? nullptr : DMXControlledProperties[0]->GetFixturePatch();
 	}
 
-	const UObject* FRemoteControlDMXControlledPropertyPatch::GetOwnerObject() const
+	const UObject* FRemoteControlDMXControlledPropertyPatch::GetOwnerActor() const
 	{
-		return DMXControlledProperties.IsEmpty() ? nullptr : DMXControlledProperties[0]->GetOwnerObject();
+		return DMXControlledProperties.IsEmpty() ? nullptr : DMXControlledProperties[0]->GetOwnerActor();
 	}
 
 #if WITH_EDITOR
diff --git a/Engine/Plugins/VirtualProduction/RemoteControlProtocolDMX/Source/RemoteControlProtocolDMXEditor/Private/Library/RemoteControlDMXLibraryBuilder.cpp b/Engine/Plugins/VirtualProduction/RemoteControlProtocolDMX/Source/RemoteControlProtocolDMXEditor/Private/Library/RemoteControlDMXLibraryBuilder.cpp
--- a/Engine/Plugins/VirtualProduction/RemoteControlProtocolDMX/Source/RemoteControlProtocolDMXEditor/Private/Library/RemoteControlDMXLibraryBuilder.cpp
+++ b/Engine/Plugins/VirtualProduction/RemoteControlProtocolDMX/Source/RemoteControlProtocolDMXEditor/Private/Library/RemoteControlDMXLibraryBuilder.cpp
@@ -42,7 +42,7 @@ namespace UE::RemoteControl::DMX
 				[](const TSharedRef<FRemoteControlDMXControlledPropertyPatch>& PropertyPatch)
 				{
 					return
-						PropertyPatch->GetOwnerObject() &&
+						PropertyPatch->GetOwnerActor() &&
 						PropertyPatch->GetFixturePatch();
 				},
 				[](const TSharedRef<FRemoteControlDMXControlledPropertyPatch>& PropertyPatch)
diff --git a/Engine/Plugins/VirtualProduction/RemoteControlProtocolDMX/Source/RemoteControlProtocolDMXEditor/Private/Library/RemoteControlDMXPatchBuilder.cpp b/Engine/Plugins/VirtualProduction/RemoteControlProtocolDMX/Source/RemoteControlProtocolDMXEditor/Private/Library/RemoteControlDMXPatchBuilder.cpp
--- a/Engine/Plugins/VirtualProduction/RemoteControlProtocolDMX/Source/RemoteControlProtocolDMXEditor/Private/Library/RemoteControlDMXPatchBuilder.cpp
+++ b/Engine/Plugins/VirtualProduction/RemoteControlProtocolDMX/Source/RemoteControlProtocolDMXEditor/Private/Library/RemoteControlDMXPatchBuilder.cpp
@@ -89,6 +89,9 @@ namespace UE::RemoteControl::DMX
 			/** Returns the name of the owner object */
 			FString GetOwnerObjectName() const;
 
+			/** Returns the transform of the owner object or Identity if the owner object has no transform */
+			FTransform GetOwnerObjectTransform() const;
+
 			/** Cached data of the fixture patch to use if a patch needs to be regenerated */
 			FRCDMXFixturePatchData FixturePatchData;
 
@@ -409,6 +412,7 @@ namespace UE::RemoteControl::DMX
 				FixturePatchConstructionParams.FixtureTypeRef = FDMXEntityFixtureTypeRef(FixtureType);
 				FixturePatchConstructionParams.ActiveMode = 0;
 				FixturePatchConstructionParams.MVRFixtureUUID = UniqueMVRFixtureUUID;
+				FixturePatchConstructionParams.DefaultTransform = GetOwnerObjectTransform();
 
 				FixturePatch = UDMXEntityFixturePatch::CreateFixturePatchInLibrary(FixturePatchConstructionParams);
 				check(FixturePatch);
@@ -495,7 +499,7 @@ namespace UE::RemoteControl::DMX
 
 		FString FRCSinglePatchBuilder::GetOwnerObjectName() const
 		{
-			const UObject* OwnerObject = DMXControlledProperties.IsEmpty() ? nullptr : DMXControlledProperties[0]->GetOwnerObject();
+			const UObject* OwnerObject = DMXControlledProperties.IsEmpty() ? nullptr : DMXControlledProperties[0]->GetOwnerActor();
 			if (!OwnerObject)
 			{
 				return FText(LOCTEXT("ObjectNotLoadedInfo", "Object is not loaded")).ToString();
@@ -510,6 +514,37 @@ namespace UE::RemoteControl::DMX
 				return OwnerObject->GetName();
 			}
 		}
+
+		FTransform FRCSinglePatchBuilder::GetOwnerObjectTransform() const
+		{
+			const UObject* OwnerObject = DMXControlledProperties.IsEmpty() ? nullptr : DMXControlledProperties[0]->ExposedProperty->GetBoundObject();
+			if (!OwnerObject)
+			{
+				return FTransform::Identity;
+			}
+
+			if (const USceneComponent* SceneComponent = Cast<USceneComponent>(OwnerObject))
+			{
+				return SceneComponent->GetComponentTransform();
+			}
+			else if (const USceneComponent* OuterSceneComponent = OwnerObject->GetTypedOuter<USceneComponent>())
+			{
+				return OuterSceneComponent->GetComponentTransform();
+			}
+			else if (const AActor* Actor = Cast<AActor>(OwnerObject))
+			{
+				return Actor->GetTransform();
+			}
+			else if (const AActor* OuterActor = OwnerObject->GetTypedOuter<AActor>())
+			{
+				return OuterActor->GetTransform();
+			}
+			else
+			{
+				return FTransform::Identity;
+			}
+		}
+
 	}
 
 
