Remote Control DMX: Add a Generate DMX Library button

#rnx

diff --git a/Engine/Plugins/VirtualProduction/RemoteControlProtocolDMX/Source/RemoteControlProtocolDMX/Internal/RemoteControlDMXUserData.h b/Engine/Plugins/VirtualProduction/RemoteControlProtocolDMX/Source/RemoteControlProtocolDMX/Internal/RemoteControlDMXUserData.h
--- a/Engine/Plugins/VirtualProduction/RemoteControlProtocolDMX/Source/RemoteControlProtocolDMX/Internal/RemoteControlDMXUserData.h
+++ b/Engine/Plugins/VirtualProduction/RemoteControlProtocolDMX/Source/RemoteControlProtocolDMX/Internal/RemoteControlDMXUserData.h
@@ -45,12 +45,15 @@ public:
 	virtual void BeginDestroy() override;
 	//~ End UObject interface 
 
-	/** Gets or creates DMX User Data in specified preset */
+	/** Gets or creates DMX User Data in specified preset. */
 	static URemoteControlDMXUserData* GetOrCreateDMXUserData(URemoteControlPreset* InPreset);
 
 	/** Returns the DMX Library held with this user data */
 	UDMXLibrary* GetDMXLibrary() const { return DMXLibrary; }
 
+	/** Sets the DMX Library held with this user data. Only valid DMX Libraries can be set. */
+	void SetDMXLibrary(UDMXLibrary* NewDMXLibrary);
+
 	/** Returns the DMX Library Proxy */
 	URemoteControlDMXLibraryProxy* GetDMXLibraryProxy() const { return DMXLibraryProxy; }
 
diff --git a/Engine/Plugins/VirtualProduction/RemoteControlProtocolDMX/Source/RemoteControlProtocolDMX/Private/RemoteControlDMXUserData.cpp b/Engine/Plugins/VirtualProduction/RemoteControlProtocolDMX/Source/RemoteControlProtocolDMX/Private/RemoteControlDMXUserData.cpp
--- a/Engine/Plugins/VirtualProduction/RemoteControlProtocolDMX/Source/RemoteControlProtocolDMX/Private/RemoteControlDMXUserData.cpp
+++ b/Engine/Plugins/VirtualProduction/RemoteControlProtocolDMX/Source/RemoteControlProtocolDMX/Private/RemoteControlDMXUserData.cpp
@@ -91,6 +91,17 @@ URemoteControlDMXUserData* URemoteControlDMXUserData::GetOrCreateDMXUserData(URe
 	return DMXUserData;
 }
 
+void URemoteControlDMXUserData::SetDMXLibrary(UDMXLibrary* NewDMXLibrary)
+{
+	if (ensureMsgf(NewDMXLibrary, TEXT("URemoteControlDMXUserData::SetDMXLibrary should not be called with null DMX Libraries. Ignoring call")))
+	{
+		DMXLibrary = NewDMXLibrary;
+
+		check(DMXLibraryProxy);
+		DMXLibraryProxy->RequestRefresh();
+	}
+}
+
 URemoteControlPreset* URemoteControlDMXUserData::GetPreset() const
 {
 	return Cast<URemoteControlPreset>(GetOuter());
diff --git a/Engine/Plugins/VirtualProduction/RemoteControlProtocolDMX/Source/RemoteControlProtocolDMXEditor/Private/UI/SRemoteControlDMXPresetUserData.cpp b/Engine/Plugins/VirtualProduction/RemoteControlProtocolDMX/Source/RemoteControlProtocolDMXEditor/Private/UI/SRemoteControlDMXPresetUserData.cpp
--- a/Engine/Plugins/VirtualProduction/RemoteControlProtocolDMX/Source/RemoteControlProtocolDMXEditor/Private/UI/SRemoteControlDMXPresetUserData.cpp
+++ b/Engine/Plugins/VirtualProduction/RemoteControlProtocolDMX/Source/RemoteControlProtocolDMXEditor/Private/UI/SRemoteControlDMXPresetUserData.cpp
@@ -2,14 +2,17 @@
 
 #include "SRemoteControlDMXPresetUserData.h"
 
+#include "AssetToolsModule.h"
 #include "Customizations/RemoteControlProtocolDMXPresetUserDataDetails.h"
 #include "DetailsViewArgs.h"
 #include "DMXEditorModule.h"
 #include "DMXEditorStyle.h"
+#include "Factories/DMXLibraryFactory.h"
 #include "IDetailsView.h"
 #include "IRCProtocolBindingList.h"
 #include "IRemoteControlProtocolWidgetsModule.h"
 #include "Library/DMXLibrary.h"
+#include "Library/DMXLibrary.h"
 #include "Modules/ModuleManager.h"
 #include "PropertyEditorModule.h"
 #include "RemoteControlDMXUserData.h"
@@ -17,6 +20,7 @@
 #include "ScopedTransaction.h"
 #include "Styling/AppStyle.h"
 #include "ToolMenus.h"
+#include "UObject/GarbageCollection.h"
 #include "Widgets/Input/SCheckBox.h"
 #include "Widgets/Input/SEditableTextBox.h"
 #include "Widgets/Layout/SSeparator.h"
@@ -124,16 +128,15 @@ namespace UE::RemoteControl::DMX
 					.Orientation(Orient_Vertical)
 				]
 
-				// Export as MVR button
-				+ SHorizontalBox::Slot()
+				// Create DMX Libary button
+				+SHorizontalBox::Slot()
 				.Padding(FMargin(8.f, 0.f))
 				.VAlign(VAlign_Center)
 				.FillWidth(1.f)
 				[
-					GenerateExportAsMVRButton()
+					GenerateActionsMenu()
 				]
 
-					
 				// DMXLibrary
 				+ SHorizontalBox::Slot()
 				.Padding(FMargin(8.f, 0.f))
@@ -155,33 +158,42 @@ namespace UE::RemoteControl::DMX
 		return TEXT("SRemoteControlDMXPresetUserData");
 	}
 	
-	TSharedRef<SWidget> SRemoteControlDMXPresetUserData::GenerateExportAsMVRButton()
+	TSharedRef<SWidget> SRemoteControlDMXPresetUserData::GenerateActionsMenu()
 	{
-		constexpr const TCHAR* MenuName = TEXT("RemoteControlPresetDMXUserDataMenu");
+		constexpr const TCHAR* MenuName = TEXT("RemoteControlProtocolDMXMenu");
 
 		UToolMenus* ToolMenus = UToolMenus::Get();
+		check(ToolMenus);
+
 		if (!ToolMenus->IsMenuRegistered(MenuName))
 		{
 			const FName NoParentToolbar = NAME_None;
-			const FName NoNameOverride = NAME_None;
-			const FText NoLabel;
-
-			UToolMenu* Menu = ToolMenus->RegisterMenu(MenuName, NoParentToolbar, EMultiBoxType::SlimHorizontalToolBar);
-			FToolMenuSection& Section = Menu->AddSection("DMXPresetUserDataMenu", NoLabel);
-
-			Section.AddMenuEntry
-			(
-				NoNameOverride,
-				LOCTEXT("ExportAsMVRLabel", "Export as MVR"),
-				LOCTEXT("ExportAsMVRTooltip", "Exports the Remote Control DMX Library as MVR file"),
-				FSlateIcon(FDMXEditorStyle::Get().GetStyleSetName(), "Icons.DMXLibraryToolbar.Export"),
-				FUIAction(FExecuteAction::CreateSP(this, &SRemoteControlDMXPresetUserData::OnExportAsMVRClicked)),
-				EUserInterfaceActionType::Button
-			);
+			ToolMenus->RegisterMenu(MenuName, NoParentToolbar, EMultiBoxType::SlimHorizontalToolBar);
 		}
+		
+		UToolMenu* Menu = UToolMenus::Get()->ExtendMenu(MenuName);
+		check(Menu);
+
+		const FText NoLabel;
+		FToolMenuSection& ActionsSection = Menu->AddSection("Actions", NoLabel);
+
+		FToolMenuEntry ExportMVREntry = FToolMenuEntry::InitToolBarButton(
+			"ExportMVR",
+			FUIAction(FExecuteAction::CreateSP(this, &SRemoteControlDMXPresetUserData::OnExportAsMVRClicked)),
+			LOCTEXT("ExportAsMVRLabel", "Export as MVR"),
+			LOCTEXT("ExportAsMVRTooltip", "Exports the Remote Control DMX Library as MVR file"),
+			FSlateIcon(FDMXEditorStyle::Get().GetStyleSetName(), "Icons.DMXLibraryToolbar.Export"));
+
+		ActionsSection.AddEntry(ExportMVREntry);
+
+		FToolMenuEntry GenerateDMXLibraryEntry = FToolMenuEntry::InitToolBarButton(
+			"GenerateDMXLibrary",
+			FUIAction(FExecuteAction::CreateSP(this, &SRemoteControlDMXPresetUserData::OnCreateDMXLibraryClicked)),
+			LOCTEXT("GenerateDMXLibraryLabel", "Create DMX Library"),
+			LOCTEXT("GenerateDMXLibraryTooltip", "Creates a new DMX Library asset from this Remote Control Preset"),
+			FSlateIcon(FDMXEditorStyle::Get().GetStyleSetName(), "ClassIcon.DMXLibrary"));
 
-		const FToolMenuContext Context(DMXUserData);
-		UToolMenu* Menu = ToolMenus->GenerateMenu(MenuName, Context);
+		ActionsSection.AddEntry(GenerateDMXLibraryEntry);
 
 		return ToolMenus->GenerateWidget(Menu);
 	}
@@ -199,6 +211,26 @@ namespace UE::RemoteControl::DMX
 		}
 	}
 
+	void SRemoteControlDMXPresetUserData::OnCreateDMXLibraryClicked()
+	{
+		if (DMXUserData)
+		{
+			UDMXLibraryFactory* DMXLibraryFactory = NewObject<UDMXLibraryFactory>();
+
+			// Ensure this object is not GC for the duration of CreateAssetWithDialog
+			FGCScopeGuard GCGuard;
+
+			FAssetToolsModule& AssetToolsModule = FAssetToolsModule::GetModule();
+			UObject* NewDMXLibraryObject = AssetToolsModule.Get().CreateAssetWithDialog(DMXLibraryFactory->GetSupportedClass(), DMXLibraryFactory);
+
+			// Set the DMX Library 
+			if (UDMXLibrary* NewDMXLibrary = Cast<UDMXLibrary>(NewDMXLibraryObject))
+			{
+				DMXUserData->SetDMXLibrary(NewDMXLibrary);
+			}
+		}
+	}
+
 	ECheckBoxState SRemoteControlDMXPresetUserData::GetAutoPatchCheckState() const
 	{
 		return DMXUserData && DMXUserData->IsAutoPatch() ? ECheckBoxState::Checked : ECheckBoxState::Unchecked;
diff --git a/Engine/Plugins/VirtualProduction/RemoteControlProtocolDMX/Source/RemoteControlProtocolDMXEditor/Private/UI/SRemoteControlDMXPresetUserData.h b/Engine/Plugins/VirtualProduction/RemoteControlProtocolDMX/Source/RemoteControlProtocolDMXEditor/Private/UI/SRemoteControlDMXPresetUserData.h
--- a/Engine/Plugins/VirtualProduction/RemoteControlProtocolDMX/Source/RemoteControlProtocolDMXEditor/Private/UI/SRemoteControlDMXPresetUserData.h
+++ b/Engine/Plugins/VirtualProduction/RemoteControlProtocolDMX/Source/RemoteControlProtocolDMXEditor/Private/UI/SRemoteControlDMXPresetUserData.h
@@ -29,12 +29,15 @@ namespace UE::RemoteControl::DMX
 		// End FGCObject interface
 
 	private:
-		/** Generates the Export as MVR button */
-		TSharedRef<SWidget> GenerateExportAsMVRButton();
+		/** Generates an Actions menu (e.g Export as MVR) */
+		TSharedRef<SWidget> GenerateActionsMenu();
 
 		/** Called when the export as MVR button was clicked */
 		void OnExportAsMVRClicked();
 
+		/** Called when the Create DMX Library button was clicked */
+		void OnCreateDMXLibraryClicked();
+
 		/** Returns true if the auto patch option is checked */
 		ECheckBoxState GetAutoPatchCheckState() const;
 
