Remote Control DMX: Fix an issue where the DMX Library did not properly clean up in some cases, leading to duplicate patches and ambigous patching, fix attribute number now shows after its property instead of its subobjectpath

#rnx

diff --git a/Engine/Plugins/VirtualProduction/RemoteControlProtocolDMX/Source/RemoteControlProtocolDMX/Private/Library/RemoteControlDMXControlledProperty.cpp b/Engine/Plugins/VirtualProduction/RemoteControlProtocolDMX/Source/RemoteControlProtocolDMX/Private/Library/RemoteControlDMXControlledProperty.cpp
--- a/Engine/Plugins/VirtualProduction/RemoteControlProtocolDMX/Source/RemoteControlProtocolDMX/Private/Library/RemoteControlDMXControlledProperty.cpp
+++ b/Engine/Plugins/VirtualProduction/RemoteControlProtocolDMX/Source/RemoteControlProtocolDMX/Private/Library/RemoteControlDMXControlledProperty.cpp
@@ -84,7 +84,7 @@ namespace UE::RemoteControl::DMX
 						}
 						else
 						{
-							return FString::Printf(TEXT("%s.%s%i"), *ExposedProperty->FieldPathInfo.ToString(), *GetSubobjectPath(), EntityIndex + 1);
+							return FString::Printf(TEXT("%s%i.%s"), *ExposedProperty->FieldPathInfo.ToString(), EntityIndex + 1, *GetSubobjectPath());
 						}
 					}();
 
diff --git a/Engine/Plugins/VirtualProduction/RemoteControlProtocolDMX/Source/RemoteControlProtocolDMX/Private/Library/RemoteControlDMXLibraryProxy.cpp b/Engine/Plugins/VirtualProduction/RemoteControlProtocolDMX/Source/RemoteControlProtocolDMX/Private/Library/RemoteControlDMXLibraryProxy.cpp
--- a/Engine/Plugins/VirtualProduction/RemoteControlProtocolDMX/Source/RemoteControlProtocolDMX/Private/Library/RemoteControlDMXLibraryProxy.cpp
+++ b/Engine/Plugins/VirtualProduction/RemoteControlProtocolDMX/Source/RemoteControlProtocolDMX/Private/Library/RemoteControlDMXLibraryProxy.cpp
@@ -198,7 +198,7 @@ void URemoteControlDMXLibraryProxy::UpdatePropertyPatches()
 		TMap<const UObject*, TArray<TSharedRef<FRemoteControlDMXControlledProperty>>> OwnerToPropertiesMap;
 		for (const TSharedRef<FRemoteControlDMXControlledProperty>& DMXControlledProperty : DMXControlledProperties)
 		{
-			const UObject* Owner = DMXControlledProperty->ExposedProperty->GetBoundObject();
+			const UObject* Owner = DMXControlledProperty->GetOwnerActor();
 			if (ensureMsgf(Owner, TEXT("Cannot group property by owner. Property has no valid bound object.")))
 			{
 				OwnerToPropertiesMap.FindOrAdd(Owner).Add(DMXControlledProperty);
diff --git a/Engine/Plugins/VirtualProduction/RemoteControlProtocolDMX/Source/RemoteControlProtocolDMX/Private/RemoteControlDMXUserData.cpp b/Engine/Plugins/VirtualProduction/RemoteControlProtocolDMX/Source/RemoteControlProtocolDMX/Private/RemoteControlDMXUserData.cpp
--- a/Engine/Plugins/VirtualProduction/RemoteControlProtocolDMX/Source/RemoteControlProtocolDMX/Private/RemoteControlDMXUserData.cpp
+++ b/Engine/Plugins/VirtualProduction/RemoteControlProtocolDMX/Source/RemoteControlProtocolDMX/Private/RemoteControlDMXUserData.cpp
@@ -26,7 +26,8 @@ void URemoteControlDMXUserData::PostEditChangeProperty(FPropertyChangedEvent& Pr
 	{
 		if (!DMXLibrary)
 		{
-			DMXLibrary = NewObject<UDMXLibrary>(this, "Internal", RF_Public | RF_Transactional);
+			const FName UnqiueName = MakeUniqueObjectName(this, UDMXLibrary::StaticClass(), "Internal");
+			DMXLibrary = NewObject<UDMXLibrary>(this, UnqiueName, RF_Public | RF_Transactional);
 		}
 
 		check(DMXLibraryProxy);
diff --git a/Engine/Plugins/VirtualProduction/RemoteControlProtocolDMX/Source/RemoteControlProtocolDMXEditor/Private/UI/SRemoteControlDMXPresetUserData.cpp b/Engine/Plugins/VirtualProduction/RemoteControlProtocolDMX/Source/RemoteControlProtocolDMXEditor/Private/UI/SRemoteControlDMXPresetUserData.cpp
--- a/Engine/Plugins/VirtualProduction/RemoteControlProtocolDMX/Source/RemoteControlProtocolDMXEditor/Private/UI/SRemoteControlDMXPresetUserData.cpp
+++ b/Engine/Plugins/VirtualProduction/RemoteControlProtocolDMX/Source/RemoteControlProtocolDMXEditor/Private/UI/SRemoteControlDMXPresetUserData.cpp
@@ -128,7 +128,7 @@ namespace UE::RemoteControl::DMX
 					.Orientation(Orient_Vertical)
 				]
 
-				// Create DMX Libary button
+				// Actions menu
 				+SHorizontalBox::Slot()
 				.Padding(FMargin(8.f, 0.f))
 				.VAlign(VAlign_Center)
