Remote Control DMX: Fix issues steming from the initial upgrade to using a DMX Library internally

* Remote Control Presets are no longer dirtied on load when using DMX
* Shorten attribute names, avoids issues where Control Console and some DMX consoles would print very long names
* Ambigous attribute names now enumerate instead of appending the full object path
* Common vector attributes append the vector component instead of a number

#rnx

diff --git a/Engine/Plugins/VirtualProduction/RemoteControlProtocolDMX/Source/RemoteControlProtocolDMX/Private/Library/RemoteControlDMXControlledProperty.cpp b/Engine/Plugins/VirtualProduction/RemoteControlProtocolDMX/Source/RemoteControlProtocolDMX/Private/Library/RemoteControlDMXControlledProperty.cpp
--- a/Engine/Plugins/VirtualProduction/RemoteControlProtocolDMX/Source/RemoteControlProtocolDMX/Private/Library/RemoteControlDMXControlledProperty.cpp
+++ b/Engine/Plugins/VirtualProduction/RemoteControlProtocolDMX/Source/RemoteControlProtocolDMX/Private/Library/RemoteControlDMXControlledProperty.cpp
@@ -78,16 +78,89 @@ namespace UE::RemoteControl::DMX
 			{
 				const FString AttributeName = [EntityIndex, this]()
 					{
+						const FString FieldPathInfo = ExposedProperty->FieldPathInfo.ToString();
 						if (Entities.Num() == 1)
 						{
-							return FString::Printf(TEXT("%s.%s"), *ExposedProperty->FieldPathInfo.ToString(), *GetSubobjectPath());
+							// Use the fielpath info for single entities
+							return FString::Printf(TEXT("%s"), *FieldPathInfo);
+						}
+						else if (Entities.Num() <= 3 && 
+							(FieldPathInfo.Contains(TEXT("Location")) || FieldPathInfo.Contains(TEXT("Scale"))))
+						{
+							// Use XYZW for entities containing "Location" or "Scale" and up to 3 components
+							FString VectorComponent;
+							switch (EntityIndex)
+							{
+							case 0:
+								VectorComponent = TEXT("_X");
+								break;
+							case 1:
+								VectorComponent = TEXT("_Y");
+								break;
+							case 2:
+								VectorComponent = TEXT("_Z");
+								break;
+							default:
+								checkf(0, TEXT("Expected max 4 components to create a vector-like attribute name, but got more."))
+							}
+
+							return FString::Printf(TEXT("%s%s"), *FieldPathInfo, *VectorComponent);
+						}
+						else if (Entities.Num() <= 3 && 
+							FieldPathInfo.Contains(TEXT("Rotation")))
+						{
+							// Use Yaw, Pitch, Roll for entities containing "Rotation" and up to 3 components
+							FString VectorComponent;
+							switch (EntityIndex)
+							{
+							case 0:
+								VectorComponent = TEXT("_Roll");
+								break;
+							case 1:
+								VectorComponent = TEXT("_Pitch");
+								break;
+							case 2:
+								VectorComponent = TEXT("_Yaw");
+								break;
+							default:
+								checkf(0, TEXT("Expected max 3 components to create a rotation-like attribute name, but got more."))
+							}
+
+							return FString::Printf(TEXT("%s%s"), *FieldPathInfo, *VectorComponent);
+						}
+						else if (Entities.Num() <= 4 &&
+							FieldPathInfo.Contains(TEXT("Color")))
+						{
+							// Use RGBA for entities containing "Color" and up to 4 components
+							FString ColorComponent;
+							switch (EntityIndex)
+							{
+							case 0:
+								ColorComponent = TEXT("_R");
+								break;
+							case 1:
+								ColorComponent = TEXT("_G");
+								break;
+							case 2:
+								ColorComponent = TEXT("_B");
+								break;
+							case 3:
+								ColorComponent = TEXT("_A");
+								break;
+							default:
+								checkf(0, TEXT("Expected max 4 components to create a color-like attribute name, but got more."))
+							}
+
+							return FString::Printf(TEXT("%s%s"), *FieldPathInfo, *ColorComponent);
 						}
 						else
 						{
-							return FString::Printf(TEXT("%s%i.%s"), *ExposedProperty->FieldPathInfo.ToString(), EntityIndex + 1, *GetSubobjectPath());
+							// For other multi dimensional entities append the entity index
+							return FString::Printf(TEXT("%s%i"), *FieldPathInfo, EntityIndex + 1);
 						}
 					}();
 
+
 				DMXEntity->ExtraSetting.AttributeName = *AttributeName;
 				DMXEntity->ExtraSetting.FunctionIndex = EntityIndex;
 			}
diff --git a/Engine/Plugins/VirtualProduction/RemoteControlProtocolDMX/Source/RemoteControlProtocolDMX/Private/Library/RemoteControlDMXLibraryProxy.cpp b/Engine/Plugins/VirtualProduction/RemoteControlProtocolDMX/Source/RemoteControlProtocolDMX/Private/Library/RemoteControlDMXLibraryProxy.cpp
--- a/Engine/Plugins/VirtualProduction/RemoteControlProtocolDMX/Source/RemoteControlProtocolDMX/Private/Library/RemoteControlDMXLibraryProxy.cpp
+++ b/Engine/Plugins/VirtualProduction/RemoteControlProtocolDMX/Source/RemoteControlProtocolDMX/Private/Library/RemoteControlDMXLibraryProxy.cpp
@@ -15,6 +15,7 @@
 #include "RemoteControlField.h"
 #include "RemoteControlPreset.h"
 #include "RemoteControlProtocolDMX.h"
+#include "UObject/Package.h"
 #include "UObject/UObjectGlobals.h"
 
 #if WITH_EDITOR
@@ -95,21 +96,30 @@ void URemoteControlDMXLibraryProxy::Refresh()
 
 	URemoteControlDMXUserData* DMXUserData = Cast<URemoteControlDMXUserData>(GetOuter());
 	URemoteControlPreset* Preset = DMXUserData ? Cast<URemoteControlPreset>(DMXUserData->GetOuter()) : nullptr;
+
 	if (Preset)
 	{
 		UnbindOnFixturePatchesReceived();
 
 #if WITH_EDITOR
-		OnPrePropertyPatchesChanged.Broadcast(Preset);
+		// Only refresh editor when the preset is dirty, indicating it changed
+		const bool bIsDirty = Preset && Preset->GetPackage() && Preset->GetPackage()->IsDirty();
+		if (bIsDirty)
+		{
+			OnPrePropertyPatchesChanged.Broadcast(Preset);
+		}
 #endif 
 
 		UpdatePropertyPatches();
 		
 #if WITH_EDITOR
-		OnPostPropertyPatchesChanged.Broadcast();
-		
-		// Listen to DMX related property changes of entities in editor
-		UpdateEntitiesObserver();
+		if (bIsDirty)
+		{
+			OnPostPropertyPatchesChanged.Broadcast();
+
+			// Listen to DMX related property changes of entities in editor
+			UpdateEntitiesObserver();
+		}
 #endif 
 
 		BindOnFixturePatchesReceived();
diff --git a/Engine/Plugins/VirtualProduction/RemoteControlProtocolDMX/Source/RemoteControlProtocolDMXEditor/Private/Library/RemoteControlDMXLibraryBuilder.cpp b/Engine/Plugins/VirtualProduction/RemoteControlProtocolDMX/Source/RemoteControlProtocolDMXEditor/Private/Library/RemoteControlDMXLibraryBuilder.cpp
--- a/Engine/Plugins/VirtualProduction/RemoteControlProtocolDMX/Source/RemoteControlProtocolDMXEditor/Private/Library/RemoteControlDMXLibraryBuilder.cpp
+++ b/Engine/Plugins/VirtualProduction/RemoteControlProtocolDMX/Source/RemoteControlProtocolDMXEditor/Private/Library/RemoteControlDMXLibraryBuilder.cpp
@@ -76,6 +76,24 @@ namespace UE::RemoteControl::DMX
 			{
 				AutoAssignFixturePatches(PostEditChangePropertyPatches);
 			}
+			else
+			{
+				// Auto assign new patches
+				TArray<TSharedRef<FRemoteControlDMXControlledPropertyPatch>> NewPropertyPatches;
+				Algo::TransformIf(PostEditChangePropertyPatches, NewPropertyPatches,
+					[this](const TSharedRef<FRemoteControlDMXControlledPropertyPatch>& PropertyPatch)
+					{
+						return
+							PropertyPatch->GetFixturePatch() &&
+							!PreviousFixturePatches.Contains(PropertyPatch->GetFixturePatch());
+					},
+					[this](const TSharedRef<FRemoteControlDMXControlledPropertyPatch>& PropertyPatch)
+					{
+						return PropertyPatch;
+					});
+
+				AutoAssignFixturePatches(NewPropertyPatches);
+			}
 
 			RemoveObsoleteFixturesFromDMXLibrary(PostEditChangePropertyPatches);
 		}
diff --git a/Engine/Plugins/VirtualProduction/RemoteControlProtocolDMX/Source/RemoteControlProtocolDMXEditor/Private/Library/RemoteControlDMXPatchBuilder.cpp b/Engine/Plugins/VirtualProduction/RemoteControlProtocolDMX/Source/RemoteControlProtocolDMXEditor/Private/Library/RemoteControlDMXPatchBuilder.cpp
--- a/Engine/Plugins/VirtualProduction/RemoteControlProtocolDMX/Source/RemoteControlProtocolDMXEditor/Private/Library/RemoteControlDMXPatchBuilder.cpp
+++ b/Engine/Plugins/VirtualProduction/RemoteControlProtocolDMX/Source/RemoteControlProtocolDMXEditor/Private/Library/RemoteControlDMXPatchBuilder.cpp
@@ -317,6 +317,7 @@ namespace UE::RemoteControl::DMX
 			FDMXFixtureMode NewMode;
 			NewMode.ModeName = TEXT("RemoteControl");
 
+			TMap<FName, int32> AttributeNameToCountMap;
 			for (const TSharedRef<FRemoteControlDMXControlledProperty>& Property : DMXControlledProperties)
 			{
 				for (const TSharedRef<TStructOnScope<FRemoteControlProtocolEntity>>& Entity : Property->GetEntities())
@@ -334,25 +335,34 @@ namespace UE::RemoteControl::DMX
 							return MaxFunctionPtr ? MaxFunctionPtr->GetLastChannel() + 1 : 1;
 						}();
 
-						FDMXFixtureFunction NewFunction;
-						NewFunction.FunctionName = Property->ExposedProperty->FieldPathInfo.ToString();
-						NewFunction.Attribute = DMXEntity->ExtraSetting.AttributeName;
-						NewFunction.Channel = NextFreeChannel;
-						NewFunction.DataType = DMXEntity->ExtraSetting.DataType;
-						NewFunction.bUseLSBMode = DMXEntity->ExtraSetting.bUseLSB;
+					const int32 AttributeCount = AttributeNameToCountMap.FindOrAdd(DMXEntity->ExtraSetting.AttributeName, 1)++;
+					const FName AttributeName = AttributeCount == 0 ?
+						DMXEntity->ExtraSetting.AttributeName :
+						*FString::Printf(TEXT("%s%i"), *DMXEntity->ExtraSetting.AttributeName.ToString(), AttributeCount);
 
-						// Addopt the attribute name from the DMX entity
-						NewFunction.Attribute = DMXEntity->ExtraSetting.AttributeName;
+					FDMXFixtureFunction NewFunction;
+					NewFunction.Attribute = AttributeName;
+					NewFunction.FunctionName = Property->ExposedProperty->FieldPathInfo.ToString();
+					NewFunction.Channel = NextFreeChannel;
+					NewFunction.DataType = DMXEntity->ExtraSetting.DataType;
+					NewFunction.bUseLSBMode = DMXEntity->ExtraSetting.bUseLSB;
 
-						const int32 FunctionIndex = NewMode.Functions.Add(NewFunction);
+					// Addopt the attribute name from the DMX entity
+					NewFunction.Attribute = DMXEntity->ExtraSetting.AttributeName;
 
-						// Remember the function index from the fixture function
-						DMXEntity->ExtraSetting.FunctionIndex = FunctionIndex;
+					const int32 FunctionIndex = NewMode.Functions.Add(NewFunction);
+
+					// Remember the function index from the fixture function
+					DMXEntity->ExtraSetting.FunctionIndex = FunctionIndex;
 				}
 			}
 
 			// Get or create the fixture type
-			if (!InOutFixtureType)
+			if (InOutFixtureType && !InOutFixtureType->Modes.IsEmpty())
+			{
+				InOutFixtureType->Modes[0] = NewMode;
+			}
+			else
 			{
 				// Create a new fixture type
 				FDMXEntityFixtureTypeConstructionParams FixtureTypeConstructionParams;
