Remote Control DMX: Fix remote control does not initialize a newly generated DMX Library asset in some cases

#rnx

diff --git a/Engine/Plugins/VirtualProduction/RemoteControlProtocolDMX/Source/RemoteControlProtocolDMX/Private/RemoteControlDMXUserData.cpp b/Engine/Plugins/VirtualProduction/RemoteControlProtocolDMX/Source/RemoteControlProtocolDMX/Private/RemoteControlDMXUserData.cpp
--- a/Engine/Plugins/VirtualProduction/RemoteControlProtocolDMX/Source/RemoteControlProtocolDMX/Private/RemoteControlDMXUserData.cpp
+++ b/Engine/Plugins/VirtualProduction/RemoteControlProtocolDMX/Source/RemoteControlProtocolDMX/Private/RemoteControlDMXUserData.cpp
@@ -96,6 +96,10 @@ void URemoteControlDMXUserData::SetDMXLibrary(UDMXLibrary* NewDMXLibrary)
 {
 	if (ensureMsgf(NewDMXLibrary, TEXT("URemoteControlDMXUserData::SetDMXLibrary should not be called with null DMX Libraries. Ignoring call")))
 	{
+#if WITH_EDITOR
+		Modify();
+#endif
+
 		DMXLibrary = NewDMXLibrary;
 
 		check(DMXLibraryProxy);
