DMX: Fix force deleting a DMX Library asset used in a Remote Control Preset crashes the engine

#rnx

diff --git a/Engine/Plugins/VirtualProduction/RemoteControlProtocolDMX/Source/RemoteControlProtocolDMX/Internal/Library/RemoteControlDMXLibraryProxy.h b/Engine/Plugins/VirtualProduction/RemoteControlProtocolDMX/Source/RemoteControlProtocolDMX/Internal/Library/RemoteControlDMXLibraryProxy.h
--- a/Engine/Plugins/VirtualProduction/RemoteControlProtocolDMX/Source/RemoteControlProtocolDMX/Internal/Library/RemoteControlDMXLibraryProxy.h
+++ b/Engine/Plugins/VirtualProduction/RemoteControlProtocolDMX/Source/RemoteControlProtocolDMX/Internal/Library/RemoteControlDMXLibraryProxy.h
@@ -52,6 +52,13 @@ public:
 	void Reset();
 
 #if WITH_EDITOR
+	/**
+	 * Clears Fixture Patches. This should be called after a DMX Library asset in use was force deleted.
+	 * This is due to the fact that remote control protocol entities do not follow GC as expected, being left with stale
+	 * DMX Libraries in their Fixture Patch Ref objects otherwise.
+	 */
+	void ClearFixturePatches();
+
 	/** Finds fixture patches maintained by this proxy that exceed universe size */
 	TArray<UDMXEntityFixturePatch*> FindPatchesThatExceedUniverseSize() const;
 
diff --git a/Engine/Plugins/VirtualProduction/RemoteControlProtocolDMX/Source/RemoteControlProtocolDMX/Internal/RemoteControlDMXUserData.h b/Engine/Plugins/VirtualProduction/RemoteControlProtocolDMX/Source/RemoteControlProtocolDMX/Internal/RemoteControlDMXUserData.h
--- a/Engine/Plugins/VirtualProduction/RemoteControlProtocolDMX/Source/RemoteControlProtocolDMX/Internal/RemoteControlDMXUserData.h
+++ b/Engine/Plugins/VirtualProduction/RemoteControlProtocolDMX/Source/RemoteControlProtocolDMX/Internal/RemoteControlDMXUserData.h
@@ -85,6 +85,9 @@ public:
 #endif // WITH_EDITOR
 
 private:
+	/** Tests if there is a DMX Library, creates a new one if it is null */
+	void EnsureValidDMXLibrary();
+
 #if WITH_EDITOR
 	/** Upgrades assets created before 5.5 that did not use a DMX library */
 	void TryUpgradeFromLegacy();
diff --git a/Engine/Plugins/VirtualProduction/RemoteControlProtocolDMX/Source/RemoteControlProtocolDMX/Private/Library/RemoteControlDMXLibraryProxy.cpp b/Engine/Plugins/VirtualProduction/RemoteControlProtocolDMX/Source/RemoteControlProtocolDMX/Private/Library/RemoteControlDMXLibraryProxy.cpp
--- a/Engine/Plugins/VirtualProduction/RemoteControlProtocolDMX/Source/RemoteControlProtocolDMX/Private/Library/RemoteControlDMXLibraryProxy.cpp
+++ b/Engine/Plugins/VirtualProduction/RemoteControlProtocolDMX/Source/RemoteControlProtocolDMX/Private/Library/RemoteControlDMXLibraryProxy.cpp
@@ -136,6 +136,27 @@ void URemoteControlDMXLibraryProxy::Reset()
 	PropertyPatches.Reset();
 }
 
+#if WITH_EDITOR
+void URemoteControlDMXLibraryProxy::ClearFixturePatches()
+{
+	for (const TSharedRef<FRemoteControlDMXControlledPropertyPatch>& PropertyPatch : PropertyPatches)
+	{
+		for (const TSharedRef<FRemoteControlDMXControlledProperty>& DMXControlledProperty : PropertyPatch->GetDMXControlledProperties())
+		{
+			for (const TSharedRef<TStructOnScope<FRemoteControlProtocolEntity>>& Entity : DMXControlledProperty->GetEntities())
+			{
+				FRemoteControlDMXProtocolEntity* DMXEntity = Entity->IsValid() ? Entity->Cast<FRemoteControlDMXProtocolEntity>() : nullptr;
+				if (!DMXEntity)
+				{
+					continue;
+				}
+				DMXEntity->ExtraSetting.FixturePatchReference = nullptr;
+			}
+		}
+	}
+}
+#endif 
+
 #if WITH_EDITOR
 TArray<UDMXEntityFixturePatch*> URemoteControlDMXLibraryProxy::FindPatchesThatExceedUniverseSize() const
 {
diff --git a/Engine/Plugins/VirtualProduction/RemoteControlProtocolDMX/Source/RemoteControlProtocolDMX/Private/RemoteControlDMXUserData.cpp b/Engine/Plugins/VirtualProduction/RemoteControlProtocolDMX/Source/RemoteControlProtocolDMX/Private/RemoteControlDMXUserData.cpp
--- a/Engine/Plugins/VirtualProduction/RemoteControlProtocolDMX/Source/RemoteControlProtocolDMX/Private/RemoteControlDMXUserData.cpp
+++ b/Engine/Plugins/VirtualProduction/RemoteControlProtocolDMX/Source/RemoteControlProtocolDMX/Private/RemoteControlDMXUserData.cpp
@@ -24,11 +24,8 @@ void URemoteControlDMXUserData::PostEditChangeProperty(FPropertyChangedEvent& Pr
 
 	if (PropertyChangedEvent.GetPropertyName() == GET_MEMBER_NAME_CHECKED(URemoteControlDMXUserData, DMXLibrary))
 	{
-		if (!DMXLibrary)
-		{
-			const FName UnqiueName = MakeUniqueObjectName(this, UDMXLibrary::StaticClass(), "Internal");
-			DMXLibrary = NewObject<UDMXLibrary>(this, UnqiueName, RF_Public | RF_Transactional);
-		}
+		// Ensure a valid DMX Lbirary exists in case it was cleared
+		EnsureValidDMXLibrary();
 
 		check(DMXLibraryProxy);
 		DMXLibraryProxy->RequestRefresh();
@@ -41,7 +38,11 @@ void URemoteControlDMXUserData::PostLoad()
 	Super::PostLoad();
 
 	if (!IsTemplate())
-	{
+	{		
+		// Handle cases where the DMX Library is no longer valid.
+		// This case can occur when force deleting the DMX Library and restarting the engine without saving the remote control preset. 
+		EnsureValidDMXLibrary();
+
 		check(DMXLibraryProxy);
 		DMXLibraryProxy->RequestRefresh();
 	}
@@ -148,6 +149,24 @@ void URemoteControlDMXUserData::SetPatchGroupMode(ERemoteControlDMXPatchGroupMod
 	DMXLibraryProxy->RequestRefresh();
 }
 
+void URemoteControlDMXUserData::EnsureValidDMXLibrary()
+{
+	if (!DMXLibrary)
+	{
+		const FName UnqiueName = MakeUniqueObjectName(this, UDMXLibrary::StaticClass(), "Internal");
+		DMXLibrary = NewObject<UDMXLibrary>(this, UnqiueName, RF_Public | RF_Transactional);
+
+#if WITH_EDITOR
+		// It is important to clear the fixture patches for the DMX Library proxy when the DMX Library is reset,
+		// this is to handle cases where a DMX Library was force deleted. 
+		// As of 5.5 Remote Control Protocol Entities do not follow GC as expected, leaving their 
+		// fixture patch references with stale DMX Library object pointers if not explicitly cleared here.
+		check(DMXLibraryProxy);
+		DMXLibraryProxy->ClearFixturePatches();
+#endif // WITH_EDITOR
+	}
+}
+
 #if WITH_EDITOR
 void URemoteControlDMXUserData::TryUpgradeFromLegacy()
 {
