Remote Control DMX: Fix ambigous attribute names can lead to ambigous channel mappings, and improper controls

#rnx

diff --git a/Engine/Plugins/VirtualProduction/RemoteControlProtocolDMX/Source/RemoteControlProtocolDMX/Private/Library/RemoteControlDMXControlledProperty.cpp b/Engine/Plugins/VirtualProduction/RemoteControlProtocolDMX/Source/RemoteControlProtocolDMX/Private/Library/RemoteControlDMXControlledProperty.cpp
--- a/Engine/Plugins/VirtualProduction/RemoteControlProtocolDMX/Source/RemoteControlProtocolDMX/Private/Library/RemoteControlDMXControlledProperty.cpp
+++ b/Engine/Plugins/VirtualProduction/RemoteControlProtocolDMX/Source/RemoteControlProtocolDMX/Private/Library/RemoteControlDMXControlledProperty.cpp
@@ -76,7 +76,7 @@ namespace UE::RemoteControl::DMX
 
 			if (DMXEntity->ExtraSetting.FunctionIndex == INDEX_NONE)
 			{
-				const FString AttributeName = [EntityIndex, this]()
+				const FString AttributeName = [&DMXEntity, EntityIndex, this]()
 					{
 						const FString FieldPathInfo = ExposedProperty->FieldPathInfo.ToString();
 						if (Entities.Num() == 1)
@@ -86,72 +86,75 @@ namespace UE::RemoteControl::DMX
 						}
 						else if (Entities.Num() <= 3 && 
 							(FieldPathInfo.Contains(TEXT("Location")) || FieldPathInfo.Contains(TEXT("Scale"))))
-						{
-							// Use XYZW for entities containing "Location" or "Scale" and up to 3 components
+						{							
+							// Use XYZ for entities containing "Location" or "Scale" and up to 3 components
 							FString VectorComponent;
-							switch (EntityIndex)
+							if (DMXEntity->HasMask(ERCMask::MaskA))
+							{
+								VectorComponent += TEXT("X");
+							}
+
+							if (DMXEntity->HasMask(ERCMask::MaskB))
+							{
+								VectorComponent += TEXT("Y");
+							}
+
+							if (DMXEntity->HasMask(ERCMask::MaskC))
 							{
-							case 0:
-								VectorComponent = TEXT("_X");
-								break;
-							case 1:
-								VectorComponent = TEXT("_Y");
-								break;
-							case 2:
-								VectorComponent = TEXT("_Z");
-								break;
-							default:
-								checkf(0, TEXT("Expected max 4 components to create a vector-like attribute name, but got more."))
+								VectorComponent += TEXT("Z");
 							}
 
-							return FString::Printf(TEXT("%s%s"), *FieldPathInfo, *VectorComponent);
+							return FString::Printf(TEXT("%s_%s"), *FieldPathInfo, *VectorComponent);
 						}
 						else if (Entities.Num() <= 3 && 
 							FieldPathInfo.Contains(TEXT("Rotation")))
 						{
 							// Use Yaw, Pitch, Roll for entities containing "Rotation" and up to 3 components
 							FString VectorComponent;
-							switch (EntityIndex)
+							if (DMXEntity->HasMask(ERCMask::MaskA))
+							{
+								VectorComponent += TEXT("Roll");
+							}
+
+							if (DMXEntity->HasMask(ERCMask::MaskB))
+							{
+								VectorComponent += TEXT("Pitch");
+							}
+
+							if (DMXEntity->HasMask(ERCMask::MaskC))
 							{
-							case 0:
-								VectorComponent = TEXT("_Roll");
-								break;
-							case 1:
-								VectorComponent = TEXT("_Pitch");
-								break;
-							case 2:
-								VectorComponent = TEXT("_Yaw");
-								break;
-							default:
-								checkf(0, TEXT("Expected max 3 components to create a rotation-like attribute name, but got more."))
+								VectorComponent += TEXT("Yaw");
 							}
 
-							return FString::Printf(TEXT("%s%s"), *FieldPathInfo, *VectorComponent);
+							return FString::Printf(TEXT("%s_%s"), *FieldPathInfo, *VectorComponent);
 						}
 						else if (Entities.Num() <= 4 &&
 							FieldPathInfo.Contains(TEXT("Color")))
 						{
 							// Use RGBA for entities containing "Color" and up to 4 components
 							FString ColorComponent;
-							switch (EntityIndex)
+							if (DMXEntity->HasMask(ERCMask::MaskA))
+							{
+								ColorComponent += TEXT("R");
+							}
+
+							if (DMXEntity->HasMask(ERCMask::MaskB))
+							{
+								ColorComponent += TEXT("G");
+							}
+
+							if (DMXEntity->HasMask(ERCMask::MaskC))
+							{
+								ColorComponent += TEXT("B");
+							}
+
+
+							if (DMXEntity->HasMask(ERCMask::MaskD))
 							{
-							case 0:
-								ColorComponent = TEXT("_R");
-								break;
-							case 1:
-								ColorComponent = TEXT("_G");
-								break;
-							case 2:
-								ColorComponent = TEXT("_B");
-								break;
-							case 3:
-								ColorComponent = TEXT("_A");
-								break;
-							default:
-								checkf(0, TEXT("Expected max 4 components to create a color-like attribute name, but got more."))
+								ColorComponent += TEXT("A");
 							}
 
-							return FString::Printf(TEXT("%s%s"), *FieldPathInfo, *ColorComponent);
+							return FString::Printf(TEXT("%s_%s"), *FieldPathInfo, *ColorComponent);
 						}
 						else
 						{
@@ -160,7 +163,6 @@ namespace UE::RemoteControl::DMX
 						}
 					}();
 
-
 				DMXEntity->ExtraSetting.AttributeName = *AttributeName;
 				DMXEntity->ExtraSetting.FunctionIndex = EntityIndex;
 			}
diff --git a/Engine/Plugins/VirtualProduction/RemoteControlProtocolDMX/Source/RemoteControlProtocolDMX/Private/RemoteControlProtocolDMX.cpp b/Engine/Plugins/VirtualProduction/RemoteControlProtocolDMX/Source/RemoteControlProtocolDMX/Private/RemoteControlProtocolDMX.cpp
--- a/Engine/Plugins/VirtualProduction/RemoteControlProtocolDMX/Source/RemoteControlProtocolDMX/Private/RemoteControlProtocolDMX.cpp
+++ b/Engine/Plugins/VirtualProduction/RemoteControlProtocolDMX/Source/RemoteControlProtocolDMX/Private/RemoteControlProtocolDMX.cpp
@@ -143,6 +143,13 @@ void FRemoteControlDMXProtocolEntity::UnbindDMX()
 	}
 }
 
+#if WITH_EDITOR
+void FRemoteControlDMXProtocolEntity::SetAttributeName(const FName& AttributeName)
+{
+	ExtraSetting.AttributeName = AttributeName;
+}
+#endif // WITH_EDITOR
+
 bool FRemoteControlDMXProtocolEntity::Serialize(FArchive& Ar)
 {
 	Ar.UsingCustomVersion(FRemoteControlProtocolDMXObjectVersion::GUID);
diff --git a/Engine/Plugins/VirtualProduction/RemoteControlProtocolDMX/Source/RemoteControlProtocolDMX/Public/RemoteControlProtocolDMX.h b/Engine/Plugins/VirtualProduction/RemoteControlProtocolDMX/Source/RemoteControlProtocolDMX/Public/RemoteControlProtocolDMX.h
--- a/Engine/Plugins/VirtualProduction/RemoteControlProtocolDMX/Source/RemoteControlProtocolDMX/Public/RemoteControlProtocolDMX.h
+++ b/Engine/Plugins/VirtualProduction/RemoteControlProtocolDMX/Source/RemoteControlProtocolDMX/Public/RemoteControlProtocolDMX.h
@@ -129,6 +129,16 @@ public:
 	/** Unbinds this entity from DMX */
 	void UnbindDMX();
 
+#if WITH_EDITOR
+	/** 
+	 * Sets the attribute name of the entity. 
+	 * Use with care, the attribute name needs to exist in the fixture patch's active mode for the entity to be functional.
+	 * 
+	 * @param AttributeName			The DMX attribute name the entity corresponds to.
+	 */
+	REMOTECONTROLPROTOCOLDMX_API void SetAttributeName(const FName& AttributeName);
+#endif // WITH_EDITOR
+
 	/** Called when the struct is serialized */
 	bool Serialize(FArchive& Ar);
 
diff --git a/Engine/Plugins/VirtualProduction/RemoteControlProtocolDMX/Source/RemoteControlProtocolDMXEditor/Private/Library/RemoteControlDMXPatchBuilder.cpp b/Engine/Plugins/VirtualProduction/RemoteControlProtocolDMX/Source/RemoteControlProtocolDMXEditor/Private/Library/RemoteControlDMXPatchBuilder.cpp
--- a/Engine/Plugins/VirtualProduction/RemoteControlProtocolDMX/Source/RemoteControlProtocolDMXEditor/Private/Library/RemoteControlDMXPatchBuilder.cpp
+++ b/Engine/Plugins/VirtualProduction/RemoteControlProtocolDMX/Source/RemoteControlProtocolDMXEditor/Private/Library/RemoteControlDMXPatchBuilder.cpp
@@ -336,10 +336,14 @@ namespace UE::RemoteControl::DMX
 						}();
 
 					const int32 AttributeCount = AttributeNameToCountMap.FindOrAdd(DMXEntity->ExtraSetting.AttributeName, 1)++;
-					const FName AttributeName = AttributeCount == 0 ?
-						DMXEntity->ExtraSetting.AttributeName :
-						*FString::Printf(TEXT("%s%i"), *DMXEntity->ExtraSetting.AttributeName.ToString(), AttributeCount);
+					const FName AttributeName = AttributeCount > 1 ?
+						*FString::Printf(TEXT("%s%i"), *DMXEntity->ExtraSetting.AttributeName.ToString(), AttributeCount) :
+						DMXEntity->ExtraSetting.AttributeName;
 
+					// Update the attribute name of the entity
+					DMXEntity->SetAttributeName(AttributeName);
+
+					// Create the related fixture function
 					FDMXFixtureFunction NewFunction;
 					NewFunction.Attribute = AttributeName;
 					NewFunction.FunctionName = Property->ExposedProperty->FieldPathInfo.ToString();
@@ -347,9 +351,6 @@ namespace UE::RemoteControl::DMX
 					NewFunction.DataType = DMXEntity->ExtraSetting.DataType;
 					NewFunction.bUseLSBMode = DMXEntity->ExtraSetting.bUseLSB;
 
-					// Addopt the attribute name from the DMX entity
-					NewFunction.Attribute = DMXEntity->ExtraSetting.AttributeName;
-
 					const int32 FunctionIndex = NewMode.Functions.Add(NewFunction);
 
 					// Remember the function index from the fixture function
@@ -399,7 +400,9 @@ namespace UE::RemoteControl::DMX
 
 		void FRCSinglePatchBuilder::UpdateFixturePatch(UDMXLibrary& DMXLibrary, UDMXEntityFixtureType* FixtureType, UDMXEntityFixturePatch* ReuseFixturePatch)
 		{
-			if (!FixtureType || !ReuseFixturePatch)
+			if (!FixtureType || 
+				!IsValid(FixtureType->GetParentLibrary()) || 
+				!ReuseFixturePatch)
 			{
 				ClearFixturePatches();
 			}
