RemoteControlDMX: Fix an issue where masked properties could generate an indistinguishable attribute name when created without signatures

#rnx

diff --git a/Engine/Plugins/VirtualProduction/RemoteControlProtocolDMX/Source/RemoteControlProtocolDMX/Private/Library/RemoteControlDMXControlledProperty.cpp b/Engine/Plugins/VirtualProduction/RemoteControlProtocolDMX/Source/RemoteControlProtocolDMX/Private/Library/RemoteControlDMXControlledProperty.cpp
--- a/Engine/Plugins/VirtualProduction/RemoteControlProtocolDMX/Source/RemoteControlProtocolDMX/Private/Library/RemoteControlDMXControlledProperty.cpp
+++ b/Engine/Plugins/VirtualProduction/RemoteControlProtocolDMX/Source/RemoteControlProtocolDMX/Private/Library/RemoteControlDMXControlledProperty.cpp
@@ -104,7 +104,9 @@ namespace UE::RemoteControl::DMX
 								VectorComponent += TEXT("Z");
 							}
 
-							return FString::Printf(TEXT("%s_%s"), *FieldPathInfo, *VectorComponent);
+							return VectorComponent.IsEmpty() ? 
+								*FieldPathInfo :
+								FString::Printf(TEXT("%s_%s"), *FieldPathInfo, *VectorComponent);
 						}
 						else if (Entities.Num() <= 3 && 
 							FieldPathInfo.Contains(TEXT("Rotation")))
@@ -126,7 +128,9 @@ namespace UE::RemoteControl::DMX
 								VectorComponent += TEXT("Yaw");
 							}
 
-							return FString::Printf(TEXT("%s_%s"), *FieldPathInfo, *VectorComponent);
+							return VectorComponent.IsEmpty() ?
+								*FieldPathInfo : 
+								FString::Printf(TEXT("%s_%s"), *FieldPathInfo, *VectorComponent);
 						}
 						else if (Entities.Num() <= 4 &&
 							FieldPathInfo.Contains(TEXT("Color")))
@@ -154,7 +158,9 @@ namespace UE::RemoteControl::DMX
 								ColorComponent += TEXT("A");
 							}
 
-							return FString::Printf(TEXT("%s_%s"), *FieldPathInfo, *ColorComponent);
+							return ColorComponent.IsEmpty() ?
+								*FieldPathInfo : 
+								FString::Printf(TEXT("%s_%s"), *FieldPathInfo, *ColorComponent);
 						}
 						else
 						{
