Remote Control DMX: Fix ambigous attribute names when adding DMX bindings without using signatures

diff --git a/Engine/Plugins/VirtualProduction/RemoteControlProtocolDMX/Source/RemoteControlProtocolDMXEditor/Private/Library/RemoteControlDMXPatchBuilder.cpp b/Engine/Plugins/VirtualProduction/RemoteControlProtocolDMX/Source/RemoteControlProtocolDMXEditor/Private/Library/RemoteControlDMXPatchBuilder.cpp
--- a/Engine/Plugins/VirtualProduction/RemoteControlProtocolDMX/Source/RemoteControlProtocolDMXEditor/Private/Library/RemoteControlDMXPatchBuilder.cpp
+++ b/Engine/Plugins/VirtualProduction/RemoteControlProtocolDMX/Source/RemoteControlProtocolDMXEditor/Private/Library/RemoteControlDMXPatchBuilder.cpp
@@ -335,9 +335,11 @@ namespace UE::RemoteControl::DMX
 							return MaxFunctionPtr ? MaxFunctionPtr->GetLastChannel() + 1 : 1;
 						}();
 
-					const int32 AttributeCount = AttributeNameToCountMap.FindOrAdd(DMXEntity->ExtraSetting.AttributeName, 1)++;
-					const FName AttributeName = AttributeCount > 1 ?
-						*FString::Printf(TEXT("%s%i"), *DMXEntity->ExtraSetting.AttributeName.ToString(), AttributeCount) :
+					const FName CleanAttributeName = *DMXEntity->ExtraSetting.AttributeName.GetPlainNameString();
+
+					const int32 AttributeCount = AttributeNameToCountMap.FindOrAdd(CleanAttributeName, 1)++;
+					FName AttributeName = AttributeCount > 1 ?
+						FName(CleanAttributeName, AttributeCount) :
 						DMXEntity->ExtraSetting.AttributeName;
 
 					// Update the attribute name of the entity
