Remote Control: Fixed issue with re-applying signatures clearing the DMX Library

diff --git a/Engine/Plugins/VirtualProduction/RemoteControlProtocolDMX/Source/RemoteControlProtocolDMX/Private/Library/RemoteControlDMXLibraryProxy.cpp b/Engine/Plugins/VirtualProduction/RemoteControlProtocolDMX/Source/RemoteControlProtocolDMX/Private/Library/RemoteControlDMXLibraryProxy.cpp
--- a/Engine/Plugins/VirtualProduction/RemoteControlProtocolDMX/Source/RemoteControlProtocolDMX/Private/Library/RemoteControlDMXLibraryProxy.cpp
+++ b/Engine/Plugins/VirtualProduction/RemoteControlProtocolDMX/Source/RemoteControlProtocolDMX/Private/Library/RemoteControlDMXLibraryProxy.cpp
@@ -197,7 +197,7 @@ void URemoteControlDMXLibraryProxy::UpdatePropertyPatches()
 		[ProtocolBindings = DMXProtocol->GetProtocolBindings()](const FRemoteControlProtocolBinding& InBinding)->bool
 		{
 			return InBinding.GetProtocolName() == FRemoteControlProtocolDMX::ProtocolName
-				&& ProtocolBindings.Contains(InBinding.GetRemoteControlProtocolEntityPtr());
+				&& ProtocolBindings.ContainsByPredicate(FRemoteControlProtocol::CreateProtocolComparator(InBinding.GetPropertyId()));
 		};
 
 	// Create DMX Controlled Properties
@@ -206,7 +206,7 @@ void URemoteControlDMXLibraryProxy::UpdatePropertyPatches()
 	TArray<TSharedRef<FRemoteControlDMXControlledProperty>> DMXControlledProperties;
 	Algo::TransformIf(ExposedProperties, DMXControlledProperties,
 		[&IsValidDMXBinding](const TSharedRef<FRemoteControlProperty>& ExposedProperty)
-		{			
+		{
 			const bool bValidPreset = ExposedProperty->GetOwner() != nullptr;
 			const bool bValidBoundObject = ExposedProperty->GetBoundObject() != nullptr;
 			const bool bHasValidDMXBindings = Algo::AnyOf(ExposedProperty->ProtocolBindings, IsValidDMXBinding);
